Updating virtual joysticks for the vehicle advanced bp project to the default engine virtual joysticks

diff --git a/Templates/TP_VehicleAdvBP/Config/DefaultInput.ini b/Templates/TP_VehicleAdvBP/Config/DefaultInput.ini
--- a/Templates/TP_VehicleAdvBP/Config/DefaultInput.ini
+++ b/Templates/TP_VehicleAdvBP/Config/DefaultInput.ini
@@ -170,7 +170,6 @@ DoubleClickTime=0.200000
 +AxisMappings=(AxisName="MoveRight",Scale=1.000000,Key=ValveIndex_Right_Thumbstick_X)
 +AxisMappings=(AxisName="MoveRight",Scale=1.000000,Key=MagicLeap_Left_Trackpad_X)
 +AxisMappings=(AxisName="MoveRight",Scale=1.000000,Key=MagicLeap_Right_Trackpad_X)
-DefaultTouchInterface=/Game/Vehicle/Sedan/SedanVirtualJoysticks.SedanVirtualJoysticks
+DefaultTouchInterface=/Engine/MobileResources/HUD/DefaultVirtualJoysticks.DefaultVirtualJoysticks
 +ConsoleKeys=Tilde
 
-
