Updating vehicle template map shared folder structure. Adding new level assets to match other templates.

diff --git a/Templates/TP_VehicleAdv/Config/DefaultEngine.ini b/Templates/TP_VehicleAdv/Config/DefaultEngine.ini
--- a/Templates/TP_VehicleAdv/Config/DefaultEngine.ini
+++ b/Templates/TP_VehicleAdv/Config/DefaultEngine.ini
@@ -2,17 +2,32 @@
 bSubstepping=True
 
 [/Script/EngineSettings.GameMapsSettings]
-EditorStartupMap=/Game/VehicleAdvCPP/Maps/VehicleAdvExampleMap
+EditorStartupMap=/Game/VehicleTemplate/Maps/VehicleAdvExampleMap.VehicleAdvExampleMap
 LocalMapOptions=
 TransitionMap=
 bUseSplitscreen=True
 TwoPlayerSplitscreenLayout=Horizontal
 ThreePlayerSplitscreenLayout=FavorTop
 GameInstanceClass=/Script/Engine.GameInstance
-GameDefaultMap=/Game/VehicleAdvCPP/Maps/VehicleAdvExampleMap
+GameDefaultMap=/Game/VehicleTemplate/Maps/VehicleAdvExampleMap.VehicleAdvExampleMap
 ServerDefaultMap=/Engine/Maps/Entry
 GlobalDefaultGameMode=/Script/TP_VehicleAdv.TP_VehicleAdvGameMode
 GlobalDefaultServerGameMode=None
 
 [/Script/IOSRuntimeSettings.IOSRuntimeSettings]
 MinimumiOSVersion=IOS_14
+
+[/Script/AndroidFileServerEditor.AndroidFileServerRuntimeSettings]
+bEnablePlugin=True
+bAllowNetworkConnection=True
+SecurityToken=45384E7B49F38D16E0F05497235849BE
+bIncludeInShipping=False
+bAllowExternalStartInShipping=False
+bCompileAFSProject=False
+bUseCompression=False
+bLogFiles=False
+bReportStats=False
+ConnectionType=USBOnly
+bUseManualIPAddress=False
+ManualIPAddress=
+
diff --git a/Templates/TP_VehicleAdv/Config/DefaultInput.ini b/Templates/TP_VehicleAdv/Config/DefaultInput.ini
--- a/Templates/TP_VehicleAdv/Config/DefaultInput.ini
+++ b/Templates/TP_VehicleAdv/Config/DefaultInput.ini
@@ -143,7 +143,6 @@ DoubleClickTime=0.200000
 +AxisMappings=(AxisName="MoveRight",Scale=1.000000,Key=OculusTouch_Right_Thumbstick_X)
 +AxisMappings=(AxisName="MoveRight",Scale=1.000000,Key=ValveIndex_Left_Thumbstick_X)
 +AxisMappings=(AxisName="MoveRight",Scale=1.000000,Key=ValveIndex_Right_Thumbstick_X)
-DefaultTouchInterface=/Game/Vehicle/Sedan/SedanVirtualJoysticks.SedanVirtualJoysticks
+DefaultTouchInterface=/Game/Vehicles/Vehicle/LeftVirtualJoystickOnly.LeftVirtualJoystickOnly
 +ConsoleKeys=Tilde
 
-
diff --git a/Templates/TP_VehicleAdv/Config/TemplateDefs.ini b/Templates/TP_VehicleAdv/Config/TemplateDefs.ini
--- a/Templates/TP_VehicleAdv/Config/TemplateDefs.ini
+++ b/Templates/TP_VehicleAdv/Config/TemplateDefs.ini
@@ -2,8 +2,8 @@
 
 Categories=Games
 
-LocalizedDisplayNames=(Language="en",Text="VehicleAdvanced")
-LocalizedDescriptions=(Language="en",Text="The vehicle advanced template features a more complex vehicle than the regular vehicle template. Vehicle movement can be controlled with either keyboard, controller or virtual joystick on a touch device. The template features 2 camera views - One is positioned above and behind the vehicle, whilst the other is positioned inside. When using the inside view you can look around using the mouse or touchscreen or a HMD device. The vehicle itself features working double wishbone suspension. There is also an engine sound in this template. There is also a HUD both in and out of the car that shows the current gear and speed. The level contains a landscape with rough terrain as well as a simple track that features bumps, ramps and loop the loops.")
+LocalizedDisplayNames=(Language="en",Text="Vehicle")
+LocalizedDescriptions=(Language="en",Text="The vehicle pack features both a regular vehicle and a more complex vehicle. Vehicle movement can be controlled with either keyboard, controller or virtual joystick on a touch device. The advanced vehicle itself features working double wishbone suspension. There is also an engine sound in this template. There is also a HUD that shows the current gear and speed.")
 
 LocalizedDisplayNames=(Language="ko",Text="비히클 고급")
 LocalizedDescriptions=(Language="ko",Text="비히클 고급 템플릿은 일반 비히클 템플릿보다 복잡한 비히클이 특징입니다. 비히클 이동은 키보드, 컨트롤러, 터치 디바이스의 가상 조이스틱 등 어느 것으로도 제어 가능합니다. 템플릿에는 두 가지 카메라 뷰가 제공되는데, 하나는 비히클 상단 뒤쪽에, 다른 하나는 안쪽에 위치해 있습니다. 안쪽에 있을 때는 마우스나 터치스크린이나 HMD 디바이스를 사용하여 둘러볼 수 있습니다. 비히클 자체에는 더블 위시본 서스펜션은 물론 엔진 사운드도 구현되어 있습니다. 차량 안에도 밖에도 현재 기어와 속도를 나타내는 HUD 가 있습니다. 레벨에는 거친 지형의 랜드스케이프는 물론, 방지턱, 경사로, 간단한 루프 트랙도 포함되어 있습니다.")
@@ -48,5 +48,6 @@ ReplacementsInFiles=(Extensions=("cpp","h","ini","cs"),From="%TEMPLATENAME_LOWER
 ReplacementsInFiles=(Extensions=("cpp","h","ini","cs"),From="%TEMPLATENAME%",To="%PROJECTNAME%",bCaseSensitive=false)
 
 SortKey=ZVehicle_Adv
-SharedContentPacks=(MountName="VehicleAdv",DetailLevels=("Standard"))
-SharedContentPacks=(MountName="Geometry",DetailLevels=("Standard","High"))
\ No newline at end of file
+SharedContentPacks=(MountName="Vehicles",DetailLevels=("Standard"))
+SharedContentPacks=(MountName="Track",DetailLevels=("High"))
+SharedContentPacks=(MountName="LevelPrototyping",DetailLevels=("High"))
\ No newline at end of file
diff --git a/Templates/TP_VehicleAdv/Source/TP_VehicleAdv/TP_VehicleAdvPawn.cpp b/Templates/TP_VehicleAdv/Source/TP_VehicleAdv/TP_VehicleAdvPawn.cpp
--- a/Templates/TP_VehicleAdv/Source/TP_VehicleAdv/TP_VehicleAdvPawn.cpp
+++ b/Templates/TP_VehicleAdv/Source/TP_VehicleAdv/TP_VehicleAdvPawn.cpp
@@ -40,18 +40,19 @@ PRAGMA_DISABLE_DEPRECATION_WARNINGS
 ATP_VehicleAdvPawn::ATP_VehicleAdvPawn()
 {
 	// Car mesh
-	static ConstructorHelpers::FObjectFinder<USkeletalMesh> CarMesh(TEXT("/Game/VehicleAdv/Vehicle/Vehicle_SkelMesh.Vehicle_SkelMesh"));
+	static ConstructorHelpers::FObjectFinder<USkeletalMesh> CarMesh(TEXT("/Game/Vehicles/Vehicle/Vehicle_SkelMesh.Vehicle_SkelMesh"));
 	GetMesh()->SetSkeletalMesh(CarMesh.Object);
+	GetMesh()->SetSimulatePhysics(true);
 	
-	static ConstructorHelpers::FClassFinder<UObject> AnimBPClass(TEXT("/Game/VehicleAdv/Vehicle/VehicleAnimationBlueprint"));
+	static ConstructorHelpers::FClassFinder<UObject> AnimBPClass(TEXT("/Game/Vehicles/Vehicle/VehicleAnimationBlueprint"));
 	GetMesh()->SetAnimationMode(EAnimationMode::AnimationBlueprint);
 	GetMesh()->SetAnimInstanceClass(AnimBPClass.Class);
 
 	// Setup friction materials
-	static ConstructorHelpers::FObjectFinder<UPhysicalMaterial> SlipperyMat(TEXT("/Game/VehicleAdv/PhysicsMaterials/Slippery.Slippery"));
+	static ConstructorHelpers::FObjectFinder<UPhysicalMaterial> SlipperyMat(TEXT("/Game/Vehicles/PhysicsMaterials/Slippery.Slippery"));
 	SlipperyMaterial = SlipperyMat.Object;
 		
-	static ConstructorHelpers::FObjectFinder<UPhysicalMaterial> NonSlipperyMat(TEXT("/Game/VehicleAdv/PhysicsMaterials/NonSlippery.NonSlippery"));
+	static ConstructorHelpers::FObjectFinder<UPhysicalMaterial> NonSlipperyMat(TEXT("/Game/Vehicles/PhysicsMaterials/NonSlippery.NonSlippery"));
 	NonSlipperyMaterial = NonSlipperyMat.Object;
 
 	UChaosWheeledVehicleMovementComponent* VehicleMovement = CastChecked<UChaosWheeledVehicleMovementComponent>(GetVehicleMovement());
@@ -161,7 +162,7 @@ ATP_VehicleAdvPawn::ATP_VehicleAdvPawn()
 	InCarGear->SetupAttachment(GetMesh());
 	
 	// Setup the audio component and allocate it a sound cue
-	static ConstructorHelpers::FObjectFinder<USoundCue> SoundCue(TEXT("/Game/VehicleAdv/Sound/Engine_Loop_Cue.Engine_Loop_Cue"));
+	static ConstructorHelpers::FObjectFinder<USoundCue> SoundCue(TEXT("/Game/Vehicles/Sound/Engine_Loop_Cue.Engine_Loop_Cue"));
 	EngineSoundComponent = CreateDefaultSubobject<UAudioComponent>(TEXT("EngineSound"));
 	EngineSoundComponent->SetSound(SoundCue.Object);
 	EngineSoundComponent->SetupAttachment(GetMesh());
diff --git a/Templates/TP_VehicleAdv/manifest.json b/Templates/TP_VehicleAdv/manifest.json
--- a/Templates/TP_VehicleAdv/manifest.json
+++ b/Templates/TP_VehicleAdv/manifest.json
@@ -4,7 +4,7 @@
 	[
 		{
 			"Language": "en",
-			"Text": "Vehicle Advanced"
+			"Text": "Vehicle"
 		},
 		{
 			"Language": "ja",
@@ -23,7 +23,7 @@
 	[
 		{
 			"Language": "en",
-			"Text": "The vehicle advanced pack features a more complex vehicle than the regular vehicle template. Vehicle movement can be controlled with either keyboard, controller or virtual joystick on a touch device. The template features 2 camera views - One is positioned above and behind the vehicle, whilst the other is positioned inside. When using the inside view you can look around using the mouse or touchscreen or a HMD device. The vehicle itself features working double wishbone suspension. There is also an engine sound in this template. There is also a HUD both in and out of the car that shows the current gear and speed. The level contains a landscape with rough terrain as well as a simple track that features bumps, ramps and loop the loops."
+			"Text": "The vehicle pack features both a regular vehicle and a more complex vehicle. Vehicle movement can be controlled with either keyboard, controller or virtual joystick on a touch device. The advanced vehicle itself features working double wishbone suspension. There is also an engine sound in this template. There is also a HUD that shows the current gear and speed."
 		},
 		{
 			"Language": "ja",
diff --git a/Templates/TP_VehicleAdvBP/Config/DefaultEngine.ini b/Templates/TP_VehicleAdvBP/Config/DefaultEngine.ini
--- a/Templates/TP_VehicleAdvBP/Config/DefaultEngine.ini
+++ b/Templates/TP_VehicleAdvBP/Config/DefaultEngine.ini
@@ -5,18 +5,32 @@ GameName=TP_VehicleAdvBP
 bSubstepping=True
 
 [/Script/EngineSettings.GameMapsSettings]
-EditorStartupMap=/Game/VehicleAdvBP/Maps/VehicleAdvExampleMap
+EditorStartupMap=/Game/VehicleTemplate/Maps/VehicleAdvExampleMap.VehicleAdvExampleMap
 LocalMapOptions=
 TransitionMap=
 bUseSplitscreen=True
 TwoPlayerSplitscreenLayout=Horizontal
 ThreePlayerSplitscreenLayout=FavorTop
 GameInstanceClass=/Script/Engine.GameInstance
-GameDefaultMap=/Game/VehicleAdvBP/Maps/VehicleAdvExampleMap
-ServerDefaultMap=/Engine/Maps/Entry
-GlobalDefaultGameMode=/Game/VehicleAdvBP/Blueprints/VehicleAdvGameMode.VehicleAdvGameMode_C
+GameDefaultMap=/Game/VehicleTemplate/Maps/VehicleAdvExampleMap.VehicleAdvExampleMap
+ServerDefaultMap=/Engine/Maps/Entry.Entry
+GlobalDefaultGameMode=/Game/VehicleTemplate/Blueprints/VehicleAdvGameMode.VehicleAdvGameMode_C
 GlobalDefaultServerGameMode=None
 
 [/Script/IOSRuntimeSettings.IOSRuntimeSettings]
 MinimumiOSVersion=IOS_14
 
+[/Script/AndroidFileServerEditor.AndroidFileServerRuntimeSettings]
+bEnablePlugin=True
+bAllowNetworkConnection=True
+SecurityToken=3BE0C6B84B81CE26B8E97BBB4D1902A2
+bIncludeInShipping=False
+bAllowExternalStartInShipping=False
+bCompileAFSProject=False
+bUseCompression=False
+bLogFiles=False
+bReportStats=False
+ConnectionType=USBOnly
+bUseManualIPAddress=False
+ManualIPAddress=
+
diff --git a/Templates/TP_VehicleAdvBP/Config/TemplateDefs.ini b/Templates/TP_VehicleAdvBP/Config/TemplateDefs.ini
--- a/Templates/TP_VehicleAdvBP/Config/TemplateDefs.ini
+++ b/Templates/TP_VehicleAdvBP/Config/TemplateDefs.ini
@@ -2,8 +2,8 @@
 
 Categories=Games
 
-LocalizedDisplayNames=(Language="en",Text="VehicleAdvanced")
-LocalizedDescriptions=(Language="en",Text="The vehicle advanced template features a more complex vehicle than the regular vehicle template. Vehicle movement can be controlled with either keyboard, controller or virtual joystick on a touch device. The template features 2 camera views - One is positioned above and behind the vehicle, whilst the other is positioned inside. When using the inside view you can look around using the mouse or touchscreen or a HMD device. The vehicle itself features working double wishbone suspension. There is also an engine sound in this template. There is also a HUD both in and out of the car that shows the current gear and speed. The level contains a landscape with rough terrain as well as a simple track that features bumps, ramps and loop the loops.")
+LocalizedDisplayNames=(Language="en",Text="Vehicle")
+LocalizedDescriptions=(Language="en",Text="The vehicle pack features both a regular vehicle and a more complex vehicle. Vehicle movement can be controlled with either keyboard, controller or virtual joystick on a touch device. The advanced vehicle itself features working double wishbone suspension. There is also an engine sound in this template. There is also a HUD that shows the current gear and speed.")
 
 LocalizedDisplayNames=(Language="ko",Text="비히클 고급")
 LocalizedDescriptions=(Language="ko",Text="비히클 고급 템플릿은 일반 비히클 템플릿보다 복잡한 비히클이 특징입니다. 비히클 이동은 키보드, 컨트롤러, 터치 디바이스의 가상 조이스틱 등 어느 것으로도 제어 가능합니다. 템플릿에는 두 가지 카메라 뷰가 제공되는데, 하나는 비히클 상단 뒤쪽에, 다른 하나는 안쪽에 위치해 있습니다. 안쪽에 있을 때는 마우스나 터치스크린이나 HMD 디바이스를 사용하여 둘러볼 수 있습니다. 비히클 자체에는 더블 위시본 서스펜션은 물론 엔진 사운드도 구현되어 있습니다. 차량 안에도 밖에도 현재 기어와 속도를 나타내는 HUD 가 있습니다. 레벨에는 거친 지형의 랜드스케이프는 물론, 방지턱, 경사로, 간단한 루프 트랙도 포함되어 있습니다.")
@@ -43,6 +43,7 @@ ReplacementsInFiles=(Extensions=("cpp","h","ini","cs"),From="%TEMPLATENAME_LOWER
 ReplacementsInFiles=(Extensions=("cpp","h","ini","cs"),From="%TEMPLATENAME%",To="%PROJECTNAME%",bCaseSensitive=false)
 
 SortKey=ZVehicle_AdvBP
-SharedContentPacks=(MountName="VehicleAdv",DetailLevels=("Standard"))
-SharedContentPacks=(MountName="Geometry",DetailLevels=("Standard","High"))
+SharedContentPacks=(MountName="Vehicles",DetailLevels=("Standard"))
+SharedContentPacks=(MountName="Track",DetailLevels=("High"))
+SharedContentPacks=(MountName="LevelPrototyping",DetailLevels=("High"))
 EditDetailLevelPreference="High"
\ No newline at end of file
diff --git a/Templates/TP_VehicleAdvBP/manifest.json b/Templates/TP_VehicleAdvBP/manifest.json
--- a/Templates/TP_VehicleAdvBP/manifest.json
+++ b/Templates/TP_VehicleAdvBP/manifest.json
@@ -4,7 +4,7 @@
 	[
 		{
 			"Language": "en",
-			"Text": "Vehicle Advanced"
+			"Text": "Vehicle"
 		},
 		{
 			"Language": "ja",
@@ -23,7 +23,7 @@
 	[
 		{
 			"Language": "en",
-			"Text": "The vehicle advanced pack features a more complex vehicle than the regular vehicle template. Vehicle movement can be controlled with either keyboard, controller or virtual joystick on a touch device. The template features 2 camera views - One is positioned above and behind the vehicle, whilst the other is positioned inside. When using the inside view you can look around using the mouse or touchscreen or a HMD device. The vehicle itself features working double wishbone suspension. There is also an engine sound in this template. There is also a HUD both in and out of the car that shows the current gear and speed. The level contains a landscape with rough terrain as well as a simple track that features bumps, ramps and loop the loops."
+			"Text": "The vehicle pack features both a regular vehicle and a more complex vehicle. Vehicle movement can be controlled with either keyboard, controller or virtual joystick on a touch device. The advanced vehicle itself features working double wishbone suspension. There is also an engine sound in this template. There is also a HUD that shows the current gear and speed."
 		},
 		{
 			"Language": "ja",
