Rebuild landscape in VehicleAdvBP to fix warnings.
Update TemplateDefs for FirstPerson and Vehicle templates to resolve new project creation errors.
Enable Lumen in Vehicle templates

diff --git a/Templates/TP_FirstPersonBP/Config/TemplateDefs.ini b/Templates/TP_FirstPersonBP/Config/TemplateDefs.ini
--- a/Templates/TP_FirstPersonBP/Config/TemplateDefs.ini
+++ b/Templates/TP_FirstPersonBP/Config/TemplateDefs.ini
@@ -40,5 +40,5 @@ ReplacementsInFiles=(Extensions=("cpp","h","ini","cs"),From="%TEMPLATENAME_LOWER
 ReplacementsInFiles=(Extensions=("cpp","h","ini","cs"),From="%TEMPLATENAME%",To="%PROJECTNAME%",bCaseSensitive=false)
 SharedContentPacks=(MountName="FirstPersonArms",DetailLevels=("Standard"))
 SharedContentPacks=(MountName="FPWeapon",DetailLevels=("Standard"))
-SharedContentPacks=(MountName="LevelPrototyping",DetailLevels=("Standard","High"))
+SharedContentPacks=(MountName="LevelPrototyping",DetailLevels=("High"))
 EditDetailLevelPreference="High"
\ No newline at end of file
diff --git a/Templates/TP_VehicleAdv/Config/DefaultEngine.ini b/Templates/TP_VehicleAdv/Config/DefaultEngine.ini
--- a/Templates/TP_VehicleAdv/Config/DefaultEngine.ini
+++ b/Templates/TP_VehicleAdv/Config/DefaultEngine.ini
@@ -31,3 +31,8 @@ ConnectionType=USBOnly
 bUseManualIPAddress=False
 ManualIPAddress=
 
+[/Script/Engine.RendererSettings]
+r.ReflectionMethod=1
+r.GenerateMeshDistanceFields=True
+r.DynamicGlobalIlluminationMethod=1
+
diff --git a/Templates/TP_VehicleAdvBP/Config/DefaultEngine.ini b/Templates/TP_VehicleAdvBP/Config/DefaultEngine.ini
--- a/Templates/TP_VehicleAdvBP/Config/DefaultEngine.ini
+++ b/Templates/TP_VehicleAdvBP/Config/DefaultEngine.ini
@@ -34,3 +34,8 @@ ConnectionType=USBOnly
 bUseManualIPAddress=False
 ManualIPAddress=
 
+[/Script/Engine.RendererSettings]
+r.ReflectionMethod=1
+r.GenerateMeshDistanceFields=True
+r.DynamicGlobalIlluminationMethod=1
+
diff --git a/Templates/TP_VehicleAdvBP/Config/TemplateDefs.ini b/Templates/TP_VehicleAdvBP/Config/TemplateDefs.ini
--- a/Templates/TP_VehicleAdvBP/Config/TemplateDefs.ini
+++ b/Templates/TP_VehicleAdvBP/Config/TemplateDefs.ini
@@ -42,7 +42,7 @@ ReplacementsInFiles=(Extensions=("cpp","h","ini","cs"),From="%TEMPLATENAME_UPPER
 ReplacementsInFiles=(Extensions=("cpp","h","ini","cs"),From="%TEMPLATENAME_LOWERCASE%",To="%PROJECTNAME_LOWERCASE%",bCaseSensitive=true)
 ReplacementsInFiles=(Extensions=("cpp","h","ini","cs"),From="%TEMPLATENAME%",To="%PROJECTNAME%",bCaseSensitive=false)
 
-SortKey=ZVehicle_AdvBP
+
 SharedContentPacks=(MountName="Vehicles",DetailLevels=("Standard"))
 SharedContentPacks=(MountName="Track",DetailLevels=("High"))
 SharedContentPacks=(MountName="LevelPrototyping",DetailLevels=("High"))
diff --git a/Templates/TemplateResources/High/Track/FeaturePack/manifest.json b/Templates/TemplateResources/High/Track/FeaturePack/manifest.json
--- a/Templates/TemplateResources/High/Track/FeaturePack/manifest.json
+++ b/Templates/TemplateResources/High/Track/FeaturePack/manifest.json
@@ -35,13 +35,13 @@
 	[
 		"Track_Preview.png"
 	],
-	"Ident": "TrackStandard",
+	"Ident": "TrackHigh",
 	"AdditionalFiles":
 	{
 		"DestinationFilesFolder": "Track",
 		"AdditionalFilesList":
 		[
-			"Templates/TemplateResources/Standard/Track/Content/*.*"
+			"Templates/TemplateResources/High/Track/Content/*.*"
 		]
 	}
 }
