Implement new vehicle template functionality

diff --git a/Templates/TP_VehicleAdvBP/Config/DefaultInput.ini b/Templates/TP_VehicleAdvBP/Config/DefaultInput.ini
--- a/Templates/TP_VehicleAdvBP/Config/DefaultInput.ini
+++ b/Templates/TP_VehicleAdvBP/Config/DefaultInput.ini
@@ -2,6 +2,20 @@
 +DeviceConfigurations=(VendorID="0x046D",ProductID="0xC262",AxisProperties=((bEnabled=False),(Key=GenericUSBController_Axis1,Offset=-0.500000),(bEnabled=False),(Key=GenericUSBController_Axis2,bInverted=True,Offset=1.000000),(Key=GenericUSBController_Axis3,bInverted=True,Offset=1.000000)),ButtonProperties=((Key=GenericUSBController_Button1),(Key=GenericUSBController_Button2),(Key=GenericUSBController_Button3),(Key=GenericUSBController_Button4),(Key=GenericUSBController_Button5),(Key=GenericUSBController_Button6),(Key=GenericUSBController_Button7),(Key=GenericUSBController_Button8),(Key=GenericUSBController_Button9),(Key=GenericUSBController_Button10),(Key=GenericUSBController_Button11),(Key=GenericUSBController_Button12)))
 
 [/Script/Engine.InputSettings]
+-AxisConfig=(AxisKeyName="Gamepad_LeftX",AxisProperties=(DeadZone=0.25,Exponent=1.f,Sensitivity=1.f))
+-AxisConfig=(AxisKeyName="Gamepad_LeftY",AxisProperties=(DeadZone=0.25,Exponent=1.f,Sensitivity=1.f))
+-AxisConfig=(AxisKeyName="Gamepad_RightX",AxisProperties=(DeadZone=0.25,Exponent=1.f,Sensitivity=1.f))
+-AxisConfig=(AxisKeyName="Gamepad_RightY",AxisProperties=(DeadZone=0.25,Exponent=1.f,Sensitivity=1.f))
+-AxisConfig=(AxisKeyName="MouseX",AxisProperties=(DeadZone=0.f,Exponent=1.f,Sensitivity=0.07f))
+-AxisConfig=(AxisKeyName="MouseY",AxisProperties=(DeadZone=0.f,Exponent=1.f,Sensitivity=0.07f))
+-AxisConfig=(AxisKeyName="Mouse2D",AxisProperties=(DeadZone=0.f,Exponent=1.f,Sensitivity=0.07f))
++AxisConfig=(AxisKeyName="OculusTouch_Right_Trigger",AxisProperties=(DeadZone=0.000000,Sensitivity=1.000000,Exponent=1.000000,bInvert=False))
++AxisConfig=(AxisKeyName="OculusTouch_Right_FaceButton2",AxisProperties=(DeadZone=0.000000,Sensitivity=1.000000,Exponent=1.000000,bInvert=False))
++AxisConfig=(AxisKeyName="OculusTouch_Right_IndexPointing",AxisProperties=(DeadZone=0.000000,Sensitivity=1.000000,Exponent=1.000000,bInvert=False))
++AxisConfig=(AxisKeyName="OculusTouch_Right_ThumbUp",AxisProperties=(DeadZone=0.000000,Sensitivity=1.000000,Exponent=1.000000,bInvert=False))
++AxisConfig=(AxisKeyName="OculusTouchpad_Touchpad_X",AxisProperties=(DeadZone=0.000000,Sensitivity=1.000000,Exponent=1.000000,bInvert=False))
++AxisConfig=(AxisKeyName="OculusTouchpad_Touchpad_Y",AxisProperties=(DeadZone=0.000000,Sensitivity=1.000000,Exponent=1.000000,bInvert=False))
++AxisConfig=(AxisKeyName="Mouse2D",AxisProperties=(DeadZone=0.000000,Sensitivity=0.070000,Exponent=1.000000,bInvert=False))
 +AxisConfig=(AxisKeyName="Gamepad_LeftX",AxisProperties=(DeadZone=0.250000,Sensitivity=1.000000,Exponent=1.000000,bInvert=False))
 +AxisConfig=(AxisKeyName="Gamepad_LeftY",AxisProperties=(DeadZone=0.250000,Sensitivity=1.000000,Exponent=1.000000,bInvert=False))
 +AxisConfig=(AxisKeyName="Gamepad_RightX",AxisProperties=(DeadZone=0.250000,Sensitivity=1.000000,Exponent=1.000000,bInvert=False))
@@ -76,18 +90,13 @@
 +AxisConfig=(AxisKeyName="OculusTouch_Left_ThumbUp",AxisProperties=(DeadZone=0.000000,Sensitivity=1.000000,Exponent=1.000000,bInvert=False))
 +AxisConfig=(AxisKeyName="OculusTouch_Right_Thumbstick",AxisProperties=(DeadZone=0.000000,Sensitivity=1.000000,Exponent=1.000000,bInvert=False))
 +AxisConfig=(AxisKeyName="OculusTouch_Right_FaceButton1",AxisProperties=(DeadZone=0.000000,Sensitivity=1.000000,Exponent=1.000000,bInvert=False))
-+AxisConfig=(AxisKeyName="OculusTouch_Right_Trigger",AxisProperties=(DeadZone=0.000000,Sensitivity=1.000000,Exponent=1.000000,bInvert=False))
-+AxisConfig=(AxisKeyName="OculusTouch_Right_FaceButton2",AxisProperties=(DeadZone=0.000000,Sensitivity=1.000000,Exponent=1.000000,bInvert=False))
-+AxisConfig=(AxisKeyName="OculusTouch_Right_IndexPointing",AxisProperties=(DeadZone=0.000000,Sensitivity=1.000000,Exponent=1.000000,bInvert=False))
-+AxisConfig=(AxisKeyName="OculusTouch_Right_ThumbUp",AxisProperties=(DeadZone=0.000000,Sensitivity=1.000000,Exponent=1.000000,bInvert=False))
-+AxisConfig=(AxisKeyName="OculusTouchpad_Touchpad_X",AxisProperties=(DeadZone=0.000000,Sensitivity=1.000000,Exponent=1.000000,bInvert=False))
-+AxisConfig=(AxisKeyName="OculusTouchpad_Touchpad_Y",AxisProperties=(DeadZone=0.000000,Sensitivity=1.000000,Exponent=1.000000,bInvert=False))
 bAltEnterTogglesFullscreen=True
 bF11TogglesFullscreen=True
 bUseMouseForTouch=False
 bEnableMouseSmoothing=True
 bEnableFOVScaling=True
 bCaptureMouseOnLaunch=True
+bEnableLegacyInputScales=True
 bAlwaysShowTouchInterface=False
 bShowConsoleOnFourFingerTap=True
 bEnableGestureRecognizer=False
@@ -100,7 +109,7 @@ DoubleClickTime=0.200000
 +ActionMappings=(ActionName="Handbrake",bShift=False,bCtrl=False,bAlt=False,bCmd=False,Key=Gamepad_LeftShoulder)
 +ActionMappings=(ActionName="SwitchCamera",bShift=False,bCtrl=False,bAlt=False,bCmd=False,Key=Tab)
 +ActionMappings=(ActionName="ResetVR",bShift=False,bCtrl=False,bAlt=False,bCmd=False,Key=R)
-+ActionMappings=(ActionName="ResetCar",Key=BackSpace,bShift=False,bCtrl=False,bAlt=False,bCmd=False)
++ActionMappings=(ActionName="ResetCar",bShift=False,bCtrl=False,bAlt=False,bCmd=False,Key=BackSpace)
 +ActionMappings=(ActionName="SwitchCamera",bShift=False,bCtrl=False,bAlt=False,bCmd=False,Key=Gamepad_Special_Left)
 +ActionMappings=(ActionName="ResetVR",bShift=False,bCtrl=False,bAlt=False,bCmd=False,Key=Vive_Left_Menu_Click)
 +ActionMappings=(ActionName="ResetVR",bShift=False,bCtrl=False,bAlt=False,bCmd=False,Key=MixedReality_Left_Thumbstick_Click)
@@ -110,6 +119,7 @@ DoubleClickTime=0.200000
 +ActionMappings=(ActionName="Handbrake",bShift=False,bCtrl=False,bAlt=False,bCmd=False,Key=MixedReality_Right_Grip_Click)
 +ActionMappings=(ActionName="Handbrake",bShift=False,bCtrl=False,bAlt=False,bCmd=False,Key=OculusTouch_Right_Grip_Click)
 +ActionMappings=(ActionName="Handbrake",bShift=False,bCtrl=False,bAlt=False,bCmd=False,Key=ValveIndex_Right_A_Click)
++ActionMappings=(ActionName="ResetCar",bShift=False,bCtrl=False,bAlt=False,bCmd=False,Key=Gamepad_Special_Right)
 +AxisMappings=(AxisName="MoveForward",Scale=1.000000,Key=W)
 +AxisMappings=(AxisName="MoveForward",Scale=-1.000000,Key=S)
 +AxisMappings=(AxisName="MoveForward",Scale=1.000000,Key=Up)
@@ -125,9 +135,6 @@ DoubleClickTime=0.200000
 +AxisMappings=(AxisName="LookRight",Scale=1.000000,Key=MouseX)
 +AxisMappings=(AxisName="MoveForward",Scale=1.000000,Key=Gamepad_RightTriggerAxis)
 +AxisMappings=(AxisName="MoveForward",Scale=-1.000000,Key=Gamepad_LeftTriggerAxis)
-+AxisMappings=(AxisName="MoveForward",Scale=-2.000000,Key=GenericUSBController_Axis2)
-+AxisMappings=(AxisName="MoveForward",Scale=1.000000,Key=GenericUSBController_Axis3)
-+AxisMappings=(AxisName="MoveRight",Scale=3.000000,Key=GenericUSBController_Axis1)
 +AxisMappings=(AxisName="MoveForward",Scale=1.000000,Key=Vive_Right_Trigger_Axis)
 +AxisMappings=(AxisName="MoveForward",Scale=-1.000000,Key=Vive_Left_Trigger_Axis)
 +AxisMappings=(AxisName="MoveForward",Scale=1.000000,Key=MixedReality_Right_Trigger_Axis)
@@ -144,6 +151,10 @@ DoubleClickTime=0.200000
 +AxisMappings=(AxisName="MoveRight",Scale=1.000000,Key=OculusTouch_Right_Thumbstick_X)
 +AxisMappings=(AxisName="MoveRight",Scale=1.000000,Key=ValveIndex_Left_Thumbstick_X)
 +AxisMappings=(AxisName="MoveRight",Scale=1.000000,Key=ValveIndex_Right_Thumbstick_X)
++AxisMappings=(AxisName="LookRight",Scale=1.000000,Key=Gamepad_RightX)
+DefaultPlayerInputClass=/Script/Engine.PlayerInput
+DefaultInputComponentClass=/Script/Engine.InputComponent
 DefaultTouchInterface=/Engine/MobileResources/HUD/DefaultVirtualJoysticks.DefaultVirtualJoysticks
+-ConsoleKeys=Tilde
 +ConsoleKeys=Tilde
 
