Updating game template defs to correct thumbnail size in new project window.

diff --git a/Templates/TP_FirstPerson/Config/TemplateDefs.ini b/Templates/TP_FirstPerson/Config/TemplateDefs.ini
--- a/Templates/TP_FirstPerson/Config/TemplateDefs.ini
+++ b/Templates/TP_FirstPerson/Config/TemplateDefs.ini
@@ -14,6 +14,9 @@ LocalizedDescriptions=(Language="ja",Text="ファーストパーソン テンプ
 LocalizedDisplayNames=(Language="zh-Hans",Text="第一人称游戏")
 LocalizedDescriptions=(Language="zh-Hans",Text="第一人称模板包含一个第一人称玩家角色（通过一双手臂呈现）。您可以通过键盘、手柄或触摸式设备上的虚拟摇杆控制角色在关卡中四处移动。你还可以使用鼠标、手柄或触摸式设备上的虚拟摇杆转动角色视角，观察四周。此外，你还可以让角色拾取枪支，然后用鼠标、手柄或触摸式设备上的虚拟摇杆来发射球形子弹。子弹会撞飞关卡中的物理对象，还能从竞技场墙壁上弹开。")
 
+; This project uses a thumbnail as the project icon
+bThumbnailAsIcon=true
+
 ClassTypes=Actor, Character, GameMode
 AssetTypes=Animation Sequence, Animation Montage, Sound Wave, Skeletal Mesh, Animation Blueprint, Skeleton, Static Mesh.
 
diff --git a/Templates/TP_FirstPersonBP/Config/TemplateDefs.ini b/Templates/TP_FirstPersonBP/Config/TemplateDefs.ini
--- a/Templates/TP_FirstPersonBP/Config/TemplateDefs.ini
+++ b/Templates/TP_FirstPersonBP/Config/TemplateDefs.ini
@@ -14,6 +14,9 @@ LocalizedDescriptions=(Language="ja",Text="ファーストパーソンテンプ
 LocalizedDisplayNames=(Language="zh-Hans",Text="第一人称游戏")
 LocalizedDescriptions=(Language="zh-Hans",Text="第一人称模板拥有一个第一人称的玩家角色（通过一双手臂呈现）。您可以通过键盘、手柄或触摸式设备上的虚拟摇杆控制角色在关卡中四处移动。此外，您还可以使用鼠标、手柄或触摸式设备上的虚拟摇杆转动角色视角，观察四周。角色还装备了枪支，可使用鼠标、手柄或触摸式设备上的虚拟摇杆发射一颗球形子弹，并影响关卡中的一些物理对象，同时从竞技场的墙壁上弹开。")
 
+; This project uses a thumbnail as the project icon
+bThumbnailAsIcon=true
+
 ClassTypes=Character, GameMode, HUD, Actor
 AssetTypes=Animation Sequence, Animation Montage, Sound Wave, Skeletal Mesh, Animation Blueprint, Skeleton, Static Mesh, Texture.
 
diff --git a/Templates/TP_ThirdPerson/Config/TemplateDefs.ini b/Templates/TP_ThirdPerson/Config/TemplateDefs.ini
--- a/Templates/TP_ThirdPerson/Config/TemplateDefs.ini
+++ b/Templates/TP_ThirdPerson/Config/TemplateDefs.ini
@@ -15,6 +15,9 @@ LocalizedDescriptions=(Language="ja",Text="サードパーソン パックは、
 LocalizedDisplayNames=(Language="zh-Hans",Text="第三人称游戏")
 LocalizedDescriptions=(Language="zh-Hans",Text="第三人称游戏内容包包含一个可操控的角色以及一台位于角色身后斜上方的摄像机。你可以通过鼠标、键盘、手柄或触摸式设备上控制角色移动。摄像机会跟随角色，以越肩视角拍摄场景。这种视角突出主角，常用于动作及冒险类游戏。该模板的地图中还包含一些对象，可以让角色跳上跳下或四处推动。")
 
+; This project uses a thumbnail as the project icon
+bThumbnailAsIcon=true
+
 ClassTypes=Character, SpringArmComponent, GameMode
 AssetTypes=Animation Sequence, Animation Blueprint, Skeleton, Skeletal Mesh
 
diff --git a/Templates/TP_ThirdPersonBP/Config/TemplateDefs.ini b/Templates/TP_ThirdPersonBP/Config/TemplateDefs.ini
--- a/Templates/TP_ThirdPersonBP/Config/TemplateDefs.ini
+++ b/Templates/TP_ThirdPersonBP/Config/TemplateDefs.ini
@@ -14,6 +14,9 @@ LocalizedDescriptions=(Language="ja",Text="サードパーソン パックは、
 LocalizedDisplayNames=(Language="zh-Hans",Text="第三人称游戏")
 LocalizedDescriptions=(Language="zh-Hans",Text="第三人称游戏内容包包含一个可操控的角色以及一台位于角色身后斜上方的摄像机。你可以通过鼠标、键盘、手柄或触摸式设备上控制角色移动。摄像机会跟随角色，以越肩视角拍摄场景。这种视角突出主角，常用于动作及冒险类游戏。该模板的地图中还包含一些对象，可以让角色跳上跳下或四处推动。")
 
+; This project uses a thumbnail as the project icon
+bThumbnailAsIcon=true
+
 ClassTypes=Character, SpringArmComponent, GameMode
 AssetTypes=Animation Sequence, Animation Blueprint, Skeleton, Skeletal Mesh
 
diff --git a/Templates/TP_TopDown/Config/TemplateDefs.ini b/Templates/TP_TopDown/Config/TemplateDefs.ini
--- a/Templates/TP_TopDown/Config/TemplateDefs.ini
+++ b/Templates/TP_TopDown/Config/TemplateDefs.ini
@@ -15,6 +15,9 @@ LocalizedDescriptions=(Language="ja",Text="トップダウンテンプレート
 LocalizedDisplayNames=(Language="zh-Hans",Text="俯视角游戏")
 LocalizedDescriptions=(Language="zh-Hans",Text="俯视角游戏模板拥有一个可游玩角色，摄像机位于角色的后上方，但位置比第三人称视角要高得多。玩家可以使用鼠标或触摸屏点击目的地，或使用导航系统来辅助角色移动。此视角常用于动作角色扮演类游戏。")
 
+; This project uses a thumbnail as the project icon
+bThumbnailAsIcon=true
+
 ClassTypes=GameMode, Character, CameraComponent, SpringArmComponent, PlayerController
 AssetTypes=Animation Sequence, Skeletal Mesh, Animation Blueprint,Skeleton, Blend Space 1D
 
diff --git a/Templates/TP_TopDownBP/Config/TemplateDefs.ini b/Templates/TP_TopDownBP/Config/TemplateDefs.ini
--- a/Templates/TP_TopDownBP/Config/TemplateDefs.ini
+++ b/Templates/TP_TopDownBP/Config/TemplateDefs.ini
@@ -15,6 +15,9 @@ LocalizedDescriptions=(Language="ja",Text="トップダウン テンプレート
 LocalizedDisplayNames=(Language="zh-Hans",Text="俯视角游戏")
 LocalizedDescriptions=(Language="zh-Hans",Text="俯视角游戏模板包含一个可操控的角色以及一台位于角色后上方的摄像机。摄像机的位置比第三人称视角模板中的要高很多。玩家可以通过点击鼠标、触控板或触摸屏来选择目的地并移动角色，或持续按下鼠标来移动角色。此视角常用于动作角色扮演类游戏。")
 
+; This project uses a thumbnail as the project icon
+bThumbnailAsIcon=true
+
 ClassTypes=Character, PlayerController, CameraComponent, SpringArmComponent, GameMode
 AssetTypes=Animation Sequence, Skeletal Mesh, Animation Blueprint,Skeleton
 
diff --git a/Templates/TP_VehicleAdv/Config/TemplateDefs.ini b/Templates/TP_VehicleAdv/Config/TemplateDefs.ini
--- a/Templates/TP_VehicleAdv/Config/TemplateDefs.ini
+++ b/Templates/TP_VehicleAdv/Config/TemplateDefs.ini
@@ -14,6 +14,9 @@ LocalizedDescriptions=(Language="ja",Text="ビークル応用テンプレート
 LocalizedDisplayNames=(Language="zh-Hans",Text="高级载具类游戏")
 LocalizedDescriptions=(Language="zh-Hans",Text="高级载具类游戏模板含有一辆比普通载具类游戏模板更复杂的载具。您可通过键盘、手柄或触摸式设备上的虚拟摇杆来控制载具移动。该模板有2种摄像机视角——一种位于载具的后上方，而另一种位于载具内部。在使用内部视角时，您可通过鼠标、触摸屏或HMD设备观察四周。载具本身具有双横臂式悬架。该模板还包含引擎音效。载具内外还各有一个HUD，以显示当前装备和速度。关卡内含有一片粗糙地形，以及一条包含隆起、斜坡以及环道的简单道路。")
 
+; This project uses a thumbnail as the project icon
+bThumbnailAsIcon=true
+
 ClassTypes=GameMode, WheeledVehicle, HUD, SpringArmComponent, Font, CameraComponent, TextRenderComponent, InputComponent, PhysicalMaterial, AudioComponent
 AssetTypes=Static Mesh, Landscape Layer, Material, Instance, Physical Material, Sound Attenuation, Sound cue, Sound Wave, Texture, Skeletal Mesh, Skeleton, Physics Asset, Animation blueprint, Data Asset, Vehicle Wheel
 
diff --git a/Templates/TP_VehicleAdvBP/Config/TemplateDefs.ini b/Templates/TP_VehicleAdvBP/Config/TemplateDefs.ini
--- a/Templates/TP_VehicleAdvBP/Config/TemplateDefs.ini
+++ b/Templates/TP_VehicleAdvBP/Config/TemplateDefs.ini
@@ -14,6 +14,9 @@ LocalizedDescriptions=(Language="ja",Text="ビークル応用テンプレート
 LocalizedDisplayNames=(Language="zh-Hans",Text="高级载具类游戏")
 LocalizedDescriptions=(Language="zh-Hans",Text="高级载具类游戏模板含有一辆比普通载具类游戏模板更复杂的载具。您可通过键盘、手柄或触摸式设备上的虚拟摇杆来控制载具移动。该模板有2种摄像机视角——一种位于载具的后上方，而另一种位于载具内部。在使用内部视角时，您可通过鼠标、触摸屏或HMD设备观察四周。载具本身具有双横臂式悬架。该模板还包含引擎音效。载具内外还各有一个HUD，以显示当前装备和速度。关卡内含有一片粗糙地形，以及一条包含隆起、斜坡以及环道的简单道路。")
 
+; This project uses a thumbnail as the project icon
+bThumbnailAsIcon=true
+
 ClassTypes=GameMode, WheeledVehicle, HUD, SpringArmComponent, Font, CameraComponent, TextRenderComponent, InputComponent, PhysicalMaterial, AudioComponent
 AssetTypes=Static Mesh, Landscape Layer, Material, Instance, Physical Material, Sound Attenuation, Sound cue, Sound Wave, Texture, Skeletal Mesh, Skeleton, Physics Asset, Animation blueprint, Data Asset, Vehicle Wheel
 
