Adding materials for new Offroad Car. Removing old chaos template vehicle files.

diff --git a/Templates/TP_VehicleAdvBP/Config/DefaultEngine.ini b/Templates/TP_VehicleAdvBP/Config/DefaultEngine.ini
--- a/Templates/TP_VehicleAdvBP/Config/DefaultEngine.ini
+++ b/Templates/TP_VehicleAdvBP/Config/DefaultEngine.ini
@@ -39,3 +39,5 @@ r.ReflectionMethod=1
 r.GenerateMeshDistanceFields=True
 r.DynamicGlobalIlluminationMethod=1
 r.Mobile.EnableNoPrecomputedLightingCSMShader=1
+r.ClearCoatNormal=True
+
