[Templates] Vehicle Template new map and polishing pass

diff --git a/Templates/TP_VehicleAdvBP/Config/DefaultEngine.ini b/Templates/TP_VehicleAdvBP/Config/DefaultEngine.ini
--- a/Templates/TP_VehicleAdvBP/Config/DefaultEngine.ini
+++ b/Templates/TP_VehicleAdvBP/Config/DefaultEngine.ini
@@ -5,14 +5,14 @@ GameName=TP_VehicleAdvBP
 bSubstepping=True
 
 [/Script/EngineSettings.GameMapsSettings]
-EditorStartupMap=/Game/VehicleTemplate/Maps/VehicleAdvExampleMap.VehicleAdvExampleMap
+EditorStartupMap=/Game/VehicleTemplate/Maps/VehicleExampleMap.VehicleExampleMap
 LocalMapOptions=
 TransitionMap=
 bUseSplitscreen=True
 TwoPlayerSplitscreenLayout=Horizontal
 ThreePlayerSplitscreenLayout=FavorTop
 GameInstanceClass=/Script/Engine.GameInstance
-GameDefaultMap=/Game/VehicleTemplate/Maps/VehicleAdvExampleMap.VehicleAdvExampleMap
+GameDefaultMap=/Game/VehicleTemplate/Maps/VehicleExampleMap.VehicleExampleMap
 ServerDefaultMap=/Engine/Maps/Entry.Entry
 GlobalDefaultGameMode=/Game/VehicleTemplate/Blueprints/VehicleAdvGameMode.VehicleAdvGameMode_C
 GlobalDefaultServerGameMode=None
