Fix some localization bugs
#Rnx

diff --git a/Templates/TP_AEC_ProdConfigBP/Config/TemplateDefs.ini b/Templates/TP_AEC_ProdConfigBP/Config/TemplateDefs.ini
--- a/Templates/TP_AEC_ProdConfigBP/Config/TemplateDefs.ini
+++ b/Templates/TP_AEC_ProdConfigBP/Config/TemplateDefs.ini
@@ -15,9 +15,9 @@ LocalizedDisplayNames=(Language="ja", Text="デザインコンフィギュレー
 
 ; The localized descriptions for this template, which is displayed when the template is selected.
 LocalizedDescriptions=(Language="en", Text="This project template uses the Variant Manager, UMG and Blueprint to build a project allowing you to toggle between different object states such as visibility, trigger animation sequences, switch between views or toggle between different design options.\n\nThe user interface is dynamically created from the 3d model and the variants present in the level. This allows for users to swap the provided model for user-made assets and have the configurator's UI adapting to the new data.")
-LocalizedDescriptions=(Language="ko", Text="이 프로젝트 템플릿은 베리언트 매니저, UMG와 블루프린트를 사용한 프로젝트로 가시성,  애니메이션 시퀀스 트리거, 뷰 모드 또는 다양한 디자인 옵션 간의 전환 등과 같이 다른 오브젝트 스테이트를 전환할 수 있습니다. 유저 인터페이스는 레벨 내 3D 모델과 베리언트를 기반으로 동적으로 생성됩니다. 이를 통해 사용자들은 기본 제공된 모델을 직접 제작한 에셋과 교체하고 새로운 데이터에 맞게 컨피규레이터 UI를 업데이트할 수 있습니다.")
+LocalizedDescriptions=(Language="ko", Text="이 프로젝트 템플릿은 베리언트 매니저, UMG와 블루프린트를 사용한 프로젝트로 가시성,  애니메이션 시퀀스 트리거, 뷰 모드 또는 다양한 디자인 옵션 간의 전환 등과 같이 다른 오브젝트 스테이트를 전환할 수 있습니다.\n\n 유저 인터페이스는 레벨 내 3D 모델과 베리언트를 기반으로 동적으로 생성됩니다. 이를 통해 사용자들은 기본 제공된 모델을 직접 제작한 에셋과 교체하고 새로운 데이터에 맞게 컨피규레이터 UI를 업데이트할 수 있습니다.")
 LocalizedDescriptions=(Language="zh-Hans", Text="该项目模板使用变体管理器，UMG和蓝图构建了一个可在不同目标间切换的项目，如可见性、触发动画序列、切换视图模式或是开启/关闭不同的设计选项。\n\n用户界面是用关卡中的3D模型和变体动态创建的。因此用户能够用自制资产替换现有模型，从而使配置器的UI适配新数据。")
-LocalizedDescriptions=(Language="ja", Text="このプロジェクト テンプレートは Variant Manager、UMG、ブループリントを使用して、様々なオブジェクトの可視性の切替、アニメーション シーケンスの呼び出し、視点の変更、デザイン設定の切替ができるプロジェクトを構築するものです。ユーザーインターフェイスは、3D モデルとレベル内に表示されているバリアントから動的に作成されます。これにより、あらかじめ備わっているモデルをユーザー作成によるアセットと交換すれば、コンフィギュレーターの UI はその新たなデータに適合します。")
+LocalizedDescriptions=(Language="ja", Text="このプロジェクト テンプレートは Variant Manager、UMG、ブループリントを使用して、様々なオブジェクトの可視性の切替、アニメーション シーケンスの呼び出し、視点の変更、デザイン設定の切替ができるプロジェクトを構築するものです。\n\nユーザーインターフェイスは、3D モデルとレベル内に表示されているバリアントから動的に作成されます。これにより、あらかじめ備わっているモデルをユーザー作成によるアセットと交換すれば、コンフィギュレーターの UI はその新たなデータに適合します。")
 
 
 ; The sort key to use for the template to bypass the default alphabetical sorting.
diff --git a/Templates/TP_InCamVFXBP/Config/TemplateDefs.ini b/Templates/TP_InCamVFXBP/Config/TemplateDefs.ini
--- a/Templates/TP_InCamVFXBP/Config/TemplateDefs.ini
+++ b/Templates/TP_InCamVFXBP/Config/TemplateDefs.ini
@@ -5,7 +5,7 @@ Categories=ME
 LocalizedDisplayNames=(Language="en",Text="InCameraVFX")
 LocalizedDescriptions=(Language="en",Text="This project template provides all the necessary Blueprints, plugins, and settings for in-camera VFX workflows. Use this as a starting point for Virtual Production shoots with LED volumes. \n\nThis template includes three maps for getting started with different LED volumes: two levels that have variations of corner immersive stages, and one level that has a curved LED wall.\n\n")
 
-LocalizedDisplayNames=(Language="ko",Text="InCameraVFX")
+LocalizedDisplayNames=(Language="ko",Text="인카메라 VFX")
 LocalizedDescriptions=(Language="ko",Text="이 프로젝트 템플릿은 인카메라 VFX 워크플로에 필요한 모든 블루프린트, 플러그인 및 세팅을 제공합니다. LED 볼륨을 사용하는 버추얼 프로덕션 촬영의 시작점으로 사용해 보세요. \n\n 이 템플릿은 각종 LED 볼륨을 사용할 수 있는 3개의 맵이 있습니다. 몰입형 코너 스테이지의 베리에이션을 갖춘 레벨 2개와 곡면 LED 월을 갖춘 레벨 1개를 사용해 보세요.\n\n")
 
 LocalizedDisplayNames=(Language="ja",Text="インカメラVFX")
diff --git a/Templates/TP_VehicleAdv/Config/TemplateDefs.ini b/Templates/TP_VehicleAdv/Config/TemplateDefs.ini
--- a/Templates/TP_VehicleAdv/Config/TemplateDefs.ini
+++ b/Templates/TP_VehicleAdv/Config/TemplateDefs.ini
@@ -5,14 +5,14 @@ Categories=Games
 LocalizedDisplayNames=(Language="en",Text="Vehicle")
 LocalizedDescriptions=(Language="en",Text="The vehicle pack features both a regular vehicle and a more complex vehicle. Vehicle movement can be controlled with either keyboard, controller or virtual joystick on a touch device. The advanced vehicle itself features working double wishbone suspension. There is also an engine sound in this template. There is also a HUD that shows the current gear and speed.")
 
-LocalizedDisplayNames=(Language="ko",Text="비히클 고급")
-LocalizedDescriptions=(Language="ko",Text="비히클 고급 템플릿은 일반 비히클 템플릿보다 복잡한 비히클이 특징입니다. 비히클 이동은 키보드, 컨트롤러, 터치 디바이스의 가상 조이스틱 등 어느 것으로도 제어 가능합니다. 템플릿에는 두 가지 카메라 뷰가 제공되는데, 하나는 비히클 상단 뒤쪽에, 다른 하나는 안쪽에 위치해 있습니다. 안쪽에 있을 때는 마우스나 터치스크린이나 HMD 디바이스를 사용하여 둘러볼 수 있습니다. 비히클 자체에는 더블 위시본 서스펜션은 물론 엔진 사운드도 구현되어 있습니다. 차량 안에도 밖에도 현재 기어와 속도를 나타내는 HUD 가 있습니다. 레벨에는 거친 지형의 랜드스케이프는 물론, 방지턱, 경사로, 간단한 루프 트랙도 포함되어 있습니다.")
+LocalizedDisplayNames=(Language="ko",Text="비히클")
+LocalizedDescriptions=(Language="ko",Text="비히클 팩은 일반 비히클과 조금 더 복잡한 비히클을 모두 제공합니다. 비히클의 이동은 키보드, 컨트롤러 또는 터치 디바이스의 버추얼 조이스틱으로 제어 가능합니다. 고급 비히클은 더블 위시본 서스펜션을 갖추고 있습니다. 이 템플릿에는 엔진 사운드와, 현재 기어와 속도를 보여주는 HUD도 포함되어 있습니다.")
 
-LocalizedDisplayNames=(Language="ja",Text="ビークル応用")
-LocalizedDescriptions=(Language="ja",Text="ビークル応用テンプレートには通常のビークルテンプレートより複雑なビークルが含まれています。キーボード、コントローラー、タッチデバイスでのバーチャルジョイスティックでビークルを操作できます。このテンプレートには 2 方向からのカメラビューが用意されています。後方真上からのビューとビークル内部からのビューです。内部からのカメラを使用している時はマウス、タッチスクリーン、HMD デバイスで周りを見ることができます。このビークルにはダブルウィッシュボーン式サスペンションが含まれています。またエンジン音もテンプレートにあります。また現在のギアと速度を示す HUD も用意されています。レベルには荒地を含む地形と、隆起、坂、ループを含む簡単なレーストラックがあります。")
+LocalizedDisplayNames=(Language="ja",Text="ビークル")
+LocalizedDescriptions=(Language="ja",Text="ビークル / ビークル応用 / ビークル応用テンプレートには通常のビークルテンプレートより複雑なビークルが含まれています。キーボード、コントローラー、タッチデバイスでのバーチャルジョイスティックでビークルを操作できます。このビークルにはダブルウィッシュボーン式サスペンションが含まれています。またエンジン音もテンプレートにあります。また現在のギアと速度を示すHUDも用意されています。")
 
-LocalizedDisplayNames=(Language="zh-Hans",Text="高级载具类游戏")
-LocalizedDescriptions=(Language="zh-Hans",Text="高级载具类游戏模板含有一辆比普通载具类游戏模板更复杂的载具。您可通过键盘、手柄或触摸式设备上的虚拟摇杆来控制载具移动。该模板有2种摄像机视角——一种位于载具的后上方，而另一种位于载具内部。在使用内部视角时，您可通过鼠标、触摸屏或HMD设备观察四周。载具本身具有双横臂式悬架。该模板还包含引擎音效。载具内外还各有一个HUD，以显示当前装备和速度。关卡内含有一片粗糙地形，以及一条包含隆起、斜坡以及环道的简单道路。")
+LocalizedDisplayNames=(Language="zh-Hans",Text="载具")
+LocalizedDescriptions=(Language="zh-Hans",Text="载具包包含一辆普通载具和高级载具。你可以通过键盘、手柄或触摸式设备上的虚拟摇杆来控制载具移动。高级载具具有双横臂式悬架。此外，该模板还包含引擎音效，以及一个能够显示当前挡位和车速的HUD。")
 
 ; This project uses a thumbnail as the project icon
 bThumbnailAsIcon=true
diff --git a/Templates/TP_VehicleAdv/manifest.json b/Templates/TP_VehicleAdv/manifest.json
--- a/Templates/TP_VehicleAdv/manifest.json
+++ b/Templates/TP_VehicleAdv/manifest.json
@@ -8,15 +8,15 @@
 		},
 		{
 			"Language": "ja",
-			"Text": "ビークル応用"
+			"Text": "ビークル"
 		},
 		{
 			"Language": "zh-Hans",
-			"Text": "高级载具类游戏"
+			"Text": "载具"
 		},
 		{
 			"Language": "ko",
-			"Text": "비히클 고급"
+			"Text": "비히클"
 		}
 	],
 	"Description":
@@ -27,15 +27,15 @@
 		},
 		{
 			"Language": "ja",
-			"Text": "ビークル応用テンプレートには通常のビークルテンプレートより複雑なビークルが含まれています。キーボード、コントローラー、タッチデバイスでのバーチャルジョイスティックでビークルを操作できます。このテンプレートには 2 方向からのカメラビューが用意されています。後方真上からのビューとビークル内部からのビューです。内部からのカメラを使用している時はマウス、タッチスクリーン、HMD デバイスで周りを見ることができます。このビークルにはダブルウィッシュボーン式サスペンションが含まれています。またエンジン音もテンプレートにあります。また現在のギアと速度を示す HUD も用意されています。レベルには荒地を含む地形と、隆起、坂、ループを含む簡単なレーストラックがあります。"
+			"Text": "ビークル / ビークル応用 / ビークル応用テンプレートには通常のビークルテンプレートより複雑なビークルが含まれています。キーボード、コントローラー、タッチデバイスでのバーチャルジョイスティックでビークルを操作できます。このビークルにはダブルウィッシュボーン式サスペンションが含まれています。またエンジン音もテンプレートにあります。また現在のギアと速度を示すHUDも用意されています。"
 		},
 		{
 			"Language": "zh-Hans",
-			"Text": "高级载具类游戏模板含有一辆比普通载具类游戏模板更复杂的载具。您可通过键盘、手柄或触摸式设备上的虚拟摇杆来控制载具移动。该模板有2种摄像机视角——一种位于载具的后上方，而另一种位于载具内部。在使用内部视角时，您可通过鼠标、触摸屏或HMD设备观察四周。载具本身具有双横臂式悬架。该模板还包含引擎音效。载具内外还各有一个HUD，以显示当前装备和速度。关卡内含有一片粗糙地形，以及一条包含隆起、斜坡以及环道的简单道路。"
+			"Text": "载具包包含一辆普通载具和高级载具。你可以通过键盘、手柄或触摸式设备上的虚拟摇杆来控制载具移动。高级载具具有双横臂式悬架。此外，该模板还包含引擎音效，以及一个能够显示当前挡位和车速的HUD。"
 		},
 		{
 			"Language": "ko",
-			"Text": "비히클 고급 템플릿은 일반 비히클 템플릿보다 복잡한 비히클이 특징입니다. 비히클 이동은 키보드, 컨트롤러, 터치 디바이스의 가상 조이스틱 등 어느 것으로도 제어 가능합니다. 템플릿에는 두 가지 카메라 뷰가 제공되는데, 하나는 비히클 상단 뒤쪽에, 다른 하나는 안쪽에 위치해 있습니다. 안쪽에 있을 때는 마우스나 터치스크린이나 HMD 디바이스를 사용하여 둘러볼 수 있습니다. 비히클 자체에는 더블 위시본 서스펜션은 물론 엔진 사운드도 구현되어 있습니다. 차량 안에도 밖에도 현재 기어와 속도를 나타내는 HUD 가 있습니다. 레벨에는 거친 지형의 랜드스케이프는 물론, 방지턱, 경사로, 간단한 루프 트랙도 포함되어 있습니다."
+			"Text": "비히클 팩은 일반 비히클과 조금 더 복잡한 비히클을 모두 제공합니다. 비히클의 이동은 키보드, 컨트롤러 또는 터치 디바이스의 버추얼 조이스틱으로 제어 가능합니다. 고급 비히클은 더블 위시본 서스펜션을 갖추고 있습니다. 이 템플릿에는 엔진 사운드와, 현재 기어와 속도를 보여주는 HUD도 포함되어 있습니다."
 		}
 	],
 	"AssetTypes":
diff --git a/Templates/TP_VehicleAdvBP/Config/TemplateDefs.ini b/Templates/TP_VehicleAdvBP/Config/TemplateDefs.ini
--- a/Templates/TP_VehicleAdvBP/Config/TemplateDefs.ini
+++ b/Templates/TP_VehicleAdvBP/Config/TemplateDefs.ini
@@ -5,14 +5,14 @@ Categories=Games
 LocalizedDisplayNames=(Language="en",Text="Vehicle")
 LocalizedDescriptions=(Language="en",Text="The vehicle pack features both a regular vehicle and a more complex vehicle. Vehicle movement can be controlled with either keyboard, controller or virtual joystick on a touch device. The advanced vehicle itself features working double wishbone suspension. There is also an engine sound in this template. There is also a HUD that shows the current gear and speed.")
 
-LocalizedDisplayNames=(Language="ko",Text="비히클 고급")
-LocalizedDescriptions=(Language="ko",Text="비히클 고급 템플릿은 일반 비히클 템플릿보다 복잡한 비히클이 특징입니다. 비히클 이동은 키보드, 컨트롤러, 터치 디바이스의 가상 조이스틱 등 어느 것으로도 제어 가능합니다. 템플릿에는 두 가지 카메라 뷰가 제공되는데, 하나는 비히클 상단 뒤쪽에, 다른 하나는 안쪽에 위치해 있습니다. 안쪽에 있을 때는 마우스나 터치스크린이나 HMD 디바이스를 사용하여 둘러볼 수 있습니다. 비히클 자체에는 더블 위시본 서스펜션은 물론 엔진 사운드도 구현되어 있습니다. 차량 안에도 밖에도 현재 기어와 속도를 나타내는 HUD 가 있습니다. 레벨에는 거친 지형의 랜드스케이프는 물론, 방지턱, 경사로, 간단한 루프 트랙도 포함되어 있습니다.")
+LocalizedDisplayNames=(Language="ko",Text="비히클")
+LocalizedDescriptions=(Language="ko",Text="비히클 팩은 일반 비히클과 조금 더 복잡한 비히클을 모두 제공합니다. 비히클의 이동은 키보드, 컨트롤러 또는 터치 디바이스의 버추얼 조이스틱으로 제어 가능합니다. 고급 비히클은 더블 위시본 서스펜션을 갖추고 있습니다. 이 템플릿에는 엔진 사운드와, 현재 기어와 속도를 보여주는 HUD도 포함되어 있습니다.")
 
-LocalizedDisplayNames=(Language="ja",Text="ビークル応用")
-LocalizedDescriptions=(Language="ja",Text="ビークル応用テンプレートには通常のビークルテンプレートより複雑なビークルが含まれています。キーボード、コントローラー、タッチデバイスでのバーチャルジョイスティックでビークルを操作できます。このテンプレートには 2 方向からのカメラビューが用意されています。後方真上からのビューとビークル内部からのビューです。内部からのカメラを使用している時はマウス、タッチスクリーン、HMD デバイスで周りを見ることができます。このビークルにはダブルウィッシュボーン式サスペンションが含まれています。またエンジン音もテンプレートにあります。また現在のギアと速度を示す HUD も用意されています。レベルには荒地を含む地形と、隆起、坂、ループを含む簡単なレーストラックがあります。")
+LocalizedDisplayNames=(Language="ja",Text="ビークル")
+LocalizedDescriptions=(Language="ja",Text="ビークル / ビークル応用 / ビークル応用テンプレートには通常のビークルテンプレートより複雑なビークルが含まれています。キーボード、コントローラー、タッチデバイスでのバーチャルジョイスティックでビークルを操作できます。このビークルにはダブルウィッシュボーン式サスペンションが含まれています。またエンジン音もテンプレートにあります。また現在のギアと速度を示すHUDも用意されています。")
 
-LocalizedDisplayNames=(Language="zh-Hans",Text="高级载具类游戏")
-LocalizedDescriptions=(Language="zh-Hans",Text="高级载具类游戏模板含有一辆比普通载具类游戏模板更复杂的载具。您可通过键盘、手柄或触摸式设备上的虚拟摇杆来控制载具移动。该模板有2种摄像机视角——一种位于载具的后上方，而另一种位于载具内部。在使用内部视角时，您可通过鼠标、触摸屏或HMD设备观察四周。载具本身具有双横臂式悬架。该模板还包含引擎音效。载具内外还各有一个HUD，以显示当前装备和速度。关卡内含有一片粗糙地形，以及一条包含隆起、斜坡以及环道的简单道路。")
+LocalizedDisplayNames=(Language="zh-Hans",Text="载具")
+LocalizedDescriptions=(Language="zh-Hans",Text="载具包包含一辆普通载具和高级载具。你可以通过键盘、手柄或触摸式设备上的虚拟摇杆来控制载具移动。高级载具具有双横臂式悬架。此外，该模板还包含引擎音效，以及一个能够显示当前挡位和车速的HUD。")
 
 ; This project uses a thumbnail as the project icon
 bThumbnailAsIcon=true
diff --git a/Templates/TP_VehicleAdvBP/manifest.json b/Templates/TP_VehicleAdvBP/manifest.json
--- a/Templates/TP_VehicleAdvBP/manifest.json
+++ b/Templates/TP_VehicleAdvBP/manifest.json
@@ -8,15 +8,15 @@
 		},
 		{
 			"Language": "ja",
-			"Text": "ビークル応用"
+			"Text": "ビークル"
 		},
 		{
 			"Language": "zh-Hans",
-			"Text": "高级载具类游戏"
+			"Text": "载具"
 		},
 		{
 			"Language": "ko",
-			"Text": "비히클 고급"
+			"Text": "비히클"
 		}
 	],
 	"Description":
@@ -27,15 +27,15 @@
 		},
 		{
 			"Language": "ja",
-			"Text": "ビークル応用テンプレートには通常のビークルテンプレートより複雑なビークルが含まれています。キーボード、コントローラー、タッチデバイスでのバーチャルジョイスティックでビークルを操作できます。このテンプレートには 2 方向からのカメラビューが用意されています。後方真上からのビューとビークル内部からのビューです。内部からのカメラを使用している時はマウス、タッチスクリーン、HMD デバイスで周りを見ることができます。このビークルにはダブルウィッシュボーン式サスペンションが含まれています。またエンジン音もテンプレートにあります。また現在のギアと速度を示す HUD も用意されています。レベルには荒地を含む地形と、隆起、坂、ループを含む簡単なレーストラックがあります。"
+			"Text": "ビークル / ビークル応用 / ビークル応用テンプレートには通常のビークルテンプレートより複雑なビークルが含まれています。キーボード、コントローラー、タッチデバイスでのバーチャルジョイスティックでビークルを操作できます。このビークルにはダブルウィッシュボーン式サスペンションが含まれています。またエンジン音もテンプレートにあります。また現在のギアと速度を示すHUDも用意されています。"
 		},
 		{
 			"Language": "zh-Hans",
-			"Text": "高级载具类游戏模板含有一辆比普通载具类游戏模板更复杂的载具。您可通过键盘、手柄或触摸式设备上的虚拟摇杆来控制载具移动。该模板有2种摄像机视角——一种位于载具的后上方，而另一种位于载具内部。在使用内部视角时，您可通过鼠标、触摸屏或HMD设备观察四周。载具本身具有双横臂式悬架。该模板还包含引擎音效。载具内外还各有一个HUD，以显示当前装备和速度。关卡内含有一片粗糙地形，以及一条包含隆起、斜坡以及环道的简单道路。"
+			"Text": "载具包包含一辆普通载具和高级载具。你可以通过键盘、手柄或触摸式设备上的虚拟摇杆来控制载具移动。高级载具具有双横臂式悬架。此外，该模板还包含引擎音效，以及一个能够显示当前挡位和车速的HUD。"
 		},
 		{
 			"Language": "ko",
-			"Text": "비히클 고급 템플릿은 일반 비히클 템플릿보다 복잡한 비히클이 특징입니다. 비히클 이동은 키보드, 컨트롤러, 터치 디바이스의 가상 조이스틱 등 어느 것으로도 제어 가능합니다. 템플릿에는 두 가지 카메라 뷰가 제공되는데, 하나는 비히클 상단 뒤쪽에, 다른 하나는 안쪽에 위치해 있습니다. 안쪽에 있을 때는 마우스나 터치스크린이나 HMD 디바이스를 사용하여 둘러볼 수 있습니다. 비히클 자체에는 더블 위시본 서스펜션은 물론 엔진 사운드도 구현되어 있습니다. 차량 안에도 밖에도 현재 기어와 속도를 나타내는 HUD 가 있습니다. 레벨에는 거친 지형의 랜드스케이프는 물론, 방지턱, 경사로, 간단한 루프 트랙도 포함되어 있습니다."
+			"Text": "비히클 팩은 일반 비히클과 조금 더 복잡한 비히클을 모두 제공합니다. 비히클의 이동은 키보드, 컨트롤러 또는 터치 디바이스의 버추얼 조이스틱으로 제어 가능합니다. 고급 비히클은 더블 위시본 서스펜션을 갖추고 있습니다. 이 템플릿에는 엔진 사운드와, 현재 기어와 속도를 보여주는 HUD도 포함되어 있습니다."
 		}
 	],
 	"AssetTypes":
