updating ja description

diff --git a/Templates/TP_VehicleAdv/Config/TemplateDefs.ini b/Templates/TP_VehicleAdv/Config/TemplateDefs.ini
--- a/Templates/TP_VehicleAdv/Config/TemplateDefs.ini
+++ b/Templates/TP_VehicleAdv/Config/TemplateDefs.ini
@@ -9,7 +9,7 @@ LocalizedDisplayNames=(Language="ko",Text="비히클")
 LocalizedDescriptions=(Language="ko",Text="비히클 팩은 일반 비히클과 조금 더 복잡한 비히클을 모두 제공합니다. 비히클의 이동은 키보드, 컨트롤러 또는 터치 디바이스의 버추얼 조이스틱으로 제어 가능합니다. 고급 비히클은 더블 위시본 서스펜션을 갖추고 있습니다. 이 템플릿에는 엔진 사운드와, 현재 기어와 속도를 보여주는 HUD도 포함되어 있습니다.")
 
 LocalizedDisplayNames=(Language="ja",Text="ビークル")
-LocalizedDescriptions=(Language="ja",Text="ビークル / ビークル応用 / ビークル応用テンプレートには通常のビークルテンプレートより複雑なビークルが含まれています。キーボード、コントローラー、タッチデバイスでのバーチャルジョイスティックでビークルを操作できます。このビークルにはダブルウィッシュボーン式サスペンションが含まれています。またエンジン音もテンプレートにあります。また現在のギアと速度を示すHUDも用意されています。")
+LocalizedDescriptions=(Language="ja",Text="ビークル パックには通常のビークルとより複雑なビークルが含まれています。キーボード、コントローラー、タッチデバイスでのバーチャルジョイスティックでビークルを操作できます。このビークルにはダブルウィッシュボーン式サスペンションが含まれています。またエンジン音もテンプレートにあります。また現在のギアと速度を示すHUDも用意されています。")
 
 LocalizedDisplayNames=(Language="zh-Hans",Text="载具")
 LocalizedDescriptions=(Language="zh-Hans",Text="载具包包含一辆普通载具和高级载具。你可以通过键盘、手柄或触摸式设备上的虚拟摇杆来控制载具移动。高级载具具有双横臂式悬架。此外，该模板还包含引擎音效，以及一个能够显示当前挡位和车速的HUD。")
diff --git a/Templates/TP_VehicleAdv/manifest.json b/Templates/TP_VehicleAdv/manifest.json
--- a/Templates/TP_VehicleAdv/manifest.json
+++ b/Templates/TP_VehicleAdv/manifest.json
@@ -27,7 +27,7 @@
 		},
 		{
 			"Language": "ja",
-			"Text": "ビークル / ビークル応用 / ビークル応用テンプレートには通常のビークルテンプレートより複雑なビークルが含まれています。キーボード、コントローラー、タッチデバイスでのバーチャルジョイスティックでビークルを操作できます。このビークルにはダブルウィッシュボーン式サスペンションが含まれています。またエンジン音もテンプレートにあります。また現在のギアと速度を示すHUDも用意されています。"
+			"Text": "ビークル パックには通常のビークルとより複雑なビークルが含まれています。キーボード、コントローラー、タッチデバイスでのバーチャルジョイスティックでビークルを操作できます。このビークルにはダブルウィッシュボーン式サスペンションが含まれています。またエンジン音もテンプレートにあります。また現在のギアと速度を示すHUDも用意されています。"
 		},
 		{
 			"Language": "zh-Hans",
diff --git a/Templates/TP_VehicleAdvBP/Config/TemplateDefs.ini b/Templates/TP_VehicleAdvBP/Config/TemplateDefs.ini
--- a/Templates/TP_VehicleAdvBP/Config/TemplateDefs.ini
+++ b/Templates/TP_VehicleAdvBP/Config/TemplateDefs.ini
@@ -9,7 +9,7 @@ LocalizedDisplayNames=(Language="ko",Text="비히클")
 LocalizedDescriptions=(Language="ko",Text="비히클 팩은 일반 비히클과 조금 더 복잡한 비히클을 모두 제공합니다. 비히클의 이동은 키보드, 컨트롤러 또는 터치 디바이스의 버추얼 조이스틱으로 제어 가능합니다. 고급 비히클은 더블 위시본 서스펜션을 갖추고 있습니다. 이 템플릿에는 엔진 사운드와, 현재 기어와 속도를 보여주는 HUD도 포함되어 있습니다.")
 
 LocalizedDisplayNames=(Language="ja",Text="ビークル")
-LocalizedDescriptions=(Language="ja",Text="ビークル / ビークル応用 / ビークル応用テンプレートには通常のビークルテンプレートより複雑なビークルが含まれています。キーボード、コントローラー、タッチデバイスでのバーチャルジョイスティックでビークルを操作できます。このビークルにはダブルウィッシュボーン式サスペンションが含まれています。またエンジン音もテンプレートにあります。また現在のギアと速度を示すHUDも用意されています。")
+LocalizedDescriptions=(Language="ja",Text="ビークル パックには通常のビークルとより複雑なビークルが含まれています。キーボード、コントローラー、タッチデバイスでのバーチャルジョイスティックでビークルを操作できます。このビークルにはダブルウィッシュボーン式サスペンションが含まれています。またエンジン音もテンプレートにあります。また現在のギアと速度を示すHUDも用意されています。")
 
 LocalizedDisplayNames=(Language="zh-Hans",Text="载具")
 LocalizedDescriptions=(Language="zh-Hans",Text="载具包包含一辆普通载具和高级载具。你可以通过键盘、手柄或触摸式设备上的虚拟摇杆来控制载具移动。高级载具具有双横臂式悬架。此外，该模板还包含引擎音效，以及一个能够显示当前挡位和车速的HUD。")
diff --git a/Templates/TP_VehicleAdvBP/manifest.json b/Templates/TP_VehicleAdvBP/manifest.json
--- a/Templates/TP_VehicleAdvBP/manifest.json
+++ b/Templates/TP_VehicleAdvBP/manifest.json
@@ -27,7 +27,7 @@
 		},
 		{
 			"Language": "ja",
-			"Text": "ビークル / ビークル応用 / ビークル応用テンプレートには通常のビークルテンプレートより複雑なビークルが含まれています。キーボード、コントローラー、タッチデバイスでのバーチャルジョイスティックでビークルを操作できます。このビークルにはダブルウィッシュボーン式サスペンションが含まれています。またエンジン音もテンプレートにあります。また現在のギアと速度を示すHUDも用意されています。"
+			"Text": "ビークル パックには通常のビークルとより複雑なビークルが含まれています。キーボード、コントローラー、タッチデバイスでのバーチャルジョイスティックでビークルを操作できます。このビークルにはダブルウィッシュボーン式サスペンションが含まれています。またエンジン音もテンプレートにあります。また現在のギアと速度を示すHUDも用意されています。"
 		},
 		{
 			"Language": "zh-Hans",
