[Templates] Switching blank and advanced vehicle template to SM6. Adding AssetGuidelines to vehicle template (missing files)

#content only

diff --git a/Templates/TP_Blank/Config/DefaultEngine.ini b/Templates/TP_Blank/Config/DefaultEngine.ini
--- a/Templates/TP_Blank/Config/DefaultEngine.ini
+++ b/Templates/TP_Blank/Config/DefaultEngine.ini
@@ -3,3 +3,46 @@
 [/Script/EngineSettings.GameMapsSettings]
 GameDefaultMap=/Engine/Maps/Templates/OpenWorld
 
+[/Script/WindowsTargetPlatform.WindowsTargetSettings]
+DefaultGraphicsRHI=DefaultGraphicsRHI_Default
+-D3D12TargetedShaderFormats=PCD3D_SM5
++D3D12TargetedShaderFormats=PCD3D_SM6
+-D3D11TargetedShaderFormats=PCD3D_SM5
++D3D11TargetedShaderFormats=PCD3D_SM5
+Compiler=Default
+AudioSampleRate=48000
+AudioCallbackBufferFrameSize=1024
+AudioNumBuffersToEnqueue=1
+AudioMaxChannels=0
+AudioNumSourceWorkers=4
+SpatializationPlugin=
+SourceDataOverridePlugin=
+ReverbPlugin=
+OcclusionPlugin=
+CompressionOverrides=(bOverrideCompressionTimes=False,DurationThreshold=5.000000,MaxNumRandomBranches=0,SoundCueQualityIndex=0)
+CacheSizeKB=65536
+MaxChunkSizeOverrideKB=0
+bResampleForDevice=False
+MaxSampleRate=48000.000000
+HighSampleRate=32000.000000
+MedSampleRate=24000.000000
+LowSampleRate=12000.000000
+MinSampleRate=8000.000000
+CompressionQualityModifier=1.000000
+AutoStreamingThreshold=0.000000
+SoundCueCookQualityIndex=-1
+
+[/Script/AndroidFileServerEditor.AndroidFileServerRuntimeSettings]
+bEnablePlugin=True
+bAllowNetworkConnection=True
+SecurityToken=DED2580048DBB2163D6FA29747ACD6DB
+bIncludeInShipping=False
+bAllowExternalStartInShipping=False
+bCompileAFSProject=False
+bUseCompression=False
+bLogFiles=False
+bReportStats=False
+ConnectionType=USBOnly
+bUseManualIPAddress=False
+ManualIPAddress=
+
diff --git a/Templates/TP_BlankBP/Config/DefaultEngine.ini b/Templates/TP_BlankBP/Config/DefaultEngine.ini
--- a/Templates/TP_BlankBP/Config/DefaultEngine.ini
+++ b/Templates/TP_BlankBP/Config/DefaultEngine.ini
@@ -3,3 +3,32 @@
 [/Script/EngineSettings.GameMapsSettings]
 GameDefaultMap=/Engine/Maps/Templates/OpenWorld
 
+[/Script/WindowsTargetPlatform.WindowsTargetSettings]
+DefaultGraphicsRHI=DefaultGraphicsRHI_Default
+-D3D12TargetedShaderFormats=PCD3D_SM5
++D3D12TargetedShaderFormats=PCD3D_SM6
+-D3D11TargetedShaderFormats=PCD3D_SM5
++D3D11TargetedShaderFormats=PCD3D_SM5
+Compiler=Default
+AudioSampleRate=48000
+AudioCallbackBufferFrameSize=1024
+AudioNumBuffersToEnqueue=1
+AudioMaxChannels=0
+AudioNumSourceWorkers=4
+SpatializationPlugin=
+SourceDataOverridePlugin=
+ReverbPlugin=
+OcclusionPlugin=
+CompressionOverrides=(bOverrideCompressionTimes=False,DurationThreshold=5.000000,MaxNumRandomBranches=0,SoundCueQualityIndex=0)
+CacheSizeKB=65536
+MaxChunkSizeOverrideKB=0
+bResampleForDevice=False
+MaxSampleRate=48000.000000
+HighSampleRate=32000.000000
+MedSampleRate=24000.000000
+LowSampleRate=12000.000000
+MinSampleRate=8000.000000
+CompressionQualityModifier=1.000000
+AutoStreamingThreshold=0.000000
+SoundCueCookQualityIndex=-1
+
diff --git a/Templates/TP_VehicleAdvBP/Config/DefaultInput.ini b/Templates/TP_VehicleAdvBP/Config/DefaultInput.ini
--- a/Templates/TP_VehicleAdvBP/Config/DefaultInput.ini
+++ b/Templates/TP_VehicleAdvBP/Config/DefaultInput.ini
@@ -97,6 +97,9 @@ bEnableMouseSmoothing=True
 bEnableFOVScaling=True
 bCaptureMouseOnLaunch=True
 bEnableLegacyInputScales=True
+bEnableMotionControls=True
+bFilterInputByPlatformUser=False
+bShouldFlushPressedKeysOnViewportFocusLost=True
 bAlwaysShowTouchInterface=False
 bShowConsoleOnFourFingerTap=True
 bEnableGestureRecognizer=False
@@ -105,56 +108,10 @@ DefaultViewportMouseCaptureMode=CapturePermanently_IncludingInitialMouseDown
 DefaultViewportMouseLockMode=LockOnCapture
 FOVScale=0.011110
 DoubleClickTime=0.200000
-+ActionMappings=(ActionName="Handbrake",bShift=False,bCtrl=False,bAlt=False,bCmd=False,Key=SpaceBar)
-+ActionMappings=(ActionName="Handbrake",bShift=False,bCtrl=False,bAlt=False,bCmd=False,Key=Gamepad_LeftShoulder)
-+ActionMappings=(ActionName="SwitchCamera",bShift=False,bCtrl=False,bAlt=False,bCmd=False,Key=Tab)
-+ActionMappings=(ActionName="ResetVR",bShift=False,bCtrl=False,bAlt=False,bCmd=False,Key=R)
-+ActionMappings=(ActionName="ResetCar",bShift=False,bCtrl=False,bAlt=False,bCmd=False,Key=BackSpace)
-+ActionMappings=(ActionName="SwitchCamera",bShift=False,bCtrl=False,bAlt=False,bCmd=False,Key=Gamepad_Special_Left)
-+ActionMappings=(ActionName="ResetVR",bShift=False,bCtrl=False,bAlt=False,bCmd=False,Key=Vive_Left_Menu_Click)
-+ActionMappings=(ActionName="ResetVR",bShift=False,bCtrl=False,bAlt=False,bCmd=False,Key=MixedReality_Left_Thumbstick_Click)
-+ActionMappings=(ActionName="ResetVR",bShift=False,bCtrl=False,bAlt=False,bCmd=False,Key=OculusTouch_Left_Thumbstick_Click)
-+ActionMappings=(ActionName="ResetVR",bShift=False,bCtrl=False,bAlt=False,bCmd=False,Key=ValveIndex_Left_Thumbstick_Click)
-+ActionMappings=(ActionName="Handbrake",bShift=False,bCtrl=False,bAlt=False,bCmd=False,Key=Vive_Right_Grip_Click)
-+ActionMappings=(ActionName="Handbrake",bShift=False,bCtrl=False,bAlt=False,bCmd=False,Key=MixedReality_Right_Grip_Click)
-+ActionMappings=(ActionName="Handbrake",bShift=False,bCtrl=False,bAlt=False,bCmd=False,Key=OculusTouch_Right_Grip_Click)
-+ActionMappings=(ActionName="Handbrake",bShift=False,bCtrl=False,bAlt=False,bCmd=False,Key=ValveIndex_Right_A_Click)
-+ActionMappings=(ActionName="ResetCar",bShift=False,bCtrl=False,bAlt=False,bCmd=False,Key=Gamepad_Special_Right)
-+AxisMappings=(AxisName="MoveForward",Scale=1.000000,Key=W)
-+AxisMappings=(AxisName="MoveForward",Scale=-1.000000,Key=S)
-+AxisMappings=(AxisName="MoveForward",Scale=1.000000,Key=Up)
-+AxisMappings=(AxisName="MoveForward",Scale=-1.000000,Key=Down)
-+AxisMappings=(AxisName="MoveForward",Scale=1.000000,Key=Gamepad_FaceButton_Bottom)
-+AxisMappings=(AxisName="MoveForward",Scale=-1.000000,Key=Gamepad_RightY)
-+AxisMappings=(AxisName="MoveRight",Scale=-1.000000,Key=A)
-+AxisMappings=(AxisName="MoveRight",Scale=1.000000,Key=D)
-+AxisMappings=(AxisName="MoveRight",Scale=-1.000000,Key=Left)
-+AxisMappings=(AxisName="MoveRight",Scale=1.000000,Key=Right)
-+AxisMappings=(AxisName="MoveRight",Scale=1.000000,Key=Gamepad_LeftX)
-+AxisMappings=(AxisName="LookUp",Scale=1.000000,Key=MouseY)
-+AxisMappings=(AxisName="LookRight",Scale=1.000000,Key=MouseX)
-+AxisMappings=(AxisName="MoveForward",Scale=1.000000,Key=Gamepad_RightTriggerAxis)
-+AxisMappings=(AxisName="MoveForward",Scale=-1.000000,Key=Gamepad_LeftTriggerAxis)
-+AxisMappings=(AxisName="MoveForward",Scale=1.000000,Key=Vive_Right_Trigger_Axis)
-+AxisMappings=(AxisName="MoveForward",Scale=-1.000000,Key=Vive_Left_Trigger_Axis)
-+AxisMappings=(AxisName="MoveForward",Scale=1.000000,Key=MixedReality_Right_Trigger_Axis)
-+AxisMappings=(AxisName="MoveForward",Scale=-1.000000,Key=MixedReality_Left_Trigger_Axis)
-+AxisMappings=(AxisName="MoveForward",Scale=1.000000,Key=OculusTouch_Right_Trigger_Axis)
-+AxisMappings=(AxisName="MoveForward",Scale=-1.000000,Key=OculusTouch_Left_Trigger_Axis)
-+AxisMappings=(AxisName="MoveForward",Scale=1.000000,Key=ValveIndex_Right_Trigger_Axis)
-+AxisMappings=(AxisName="MoveForward",Scale=-1.000000,Key=ValveIndex_Left_Trigger_Axis)
-+AxisMappings=(AxisName="MoveRight",Scale=1.000000,Key=Vive_Left_Trackpad_X)
-+AxisMappings=(AxisName="MoveRight",Scale=1.000000,Key=Vive_Right_Trackpad_X)
-+AxisMappings=(AxisName="MoveRight",Scale=1.000000,Key=MixedReality_Left_Thumbstick_X)
-+AxisMappings=(AxisName="MoveRight",Scale=1.000000,Key=MixedReality_Right_Thumbstick_X)
-+AxisMappings=(AxisName="MoveRight",Scale=1.000000,Key=OculusTouch_Left_Thumbstick_X)
-+AxisMappings=(AxisName="MoveRight",Scale=1.000000,Key=OculusTouch_Right_Thumbstick_X)
-+AxisMappings=(AxisName="MoveRight",Scale=1.000000,Key=ValveIndex_Left_Thumbstick_X)
-+AxisMappings=(AxisName="MoveRight",Scale=1.000000,Key=ValveIndex_Right_Thumbstick_X)
-+AxisMappings=(AxisName="LookRight",Scale=1.000000,Key=Gamepad_RightX)
-DefaultPlayerInputClass=/Script/Engine.PlayerInput
-DefaultInputComponentClass=/Script/Engine.InputComponent
+DefaultPlayerInputClass=/Script/EnhancedInput.EnhancedPlayerInput
+DefaultInputComponentClass=/Script/EnhancedInput.EnhancedInputComponent
 DefaultTouchInterface=/Engine/MobileResources/HUD/DefaultVirtualJoysticks.DefaultVirtualJoysticks
 -ConsoleKeys=Tilde
 +ConsoleKeys=Tilde
++ConsoleKeys=Caret
 
diff --git a/Templates/TP_VehicleAdvBP/TP_VehicleAdvBP.uproject b/Templates/TP_VehicleAdvBP/TP_VehicleAdvBP.uproject
--- a/Templates/TP_VehicleAdvBP/TP_VehicleAdvBP.uproject
+++ b/Templates/TP_VehicleAdvBP/TP_VehicleAdvBP.uproject
@@ -4,10 +4,6 @@
 	"Category": "",
 	"Description": "",
 	"Plugins": [
-		{
-			"Name": "ChaosVehiclesPlugin",
-			"Enabled": true
-		},
 		{
 			"Name": "RawInput",
 			"Enabled": true
@@ -18,6 +14,10 @@
 			"TargetAllowList": [
 				"Editor"
 			]
-		}		
+		},
+		{
+			"Name": "ChaosVehiclesPlugin",
+			"Enabled": true
+		}
 	]
 }
\ No newline at end of file
