Added vehicle template updates and variants.
Updated First Person pistol anims
Reduced spring on motion for weapons
Fixed up layout on first person shooter and base map

diff --git a/Templates/TP_VehicleAdvBP/Config/DefaultEngine.ini b/Templates/TP_VehicleAdvBP/Config/DefaultEngine.ini
--- a/Templates/TP_VehicleAdvBP/Config/DefaultEngine.ini
+++ b/Templates/TP_VehicleAdvBP/Config/DefaultEngine.ini
@@ -1,21 +1,21 @@
 [URL]
-GameName=TP_VehicleAdvBP
+GameName=VehicleTemplateNew
 
 [/Script/Engine.PhysicsSettings]
 bSubstepping=False
 bSubsteppingAsync=True
 
 [/Script/EngineSettings.GameMapsSettings]
-EditorStartupMap=/Game/VehicleTemplate/Maps/VehicleExampleMap.VehicleExampleMap
+EditorStartupMap=/Game/VehicleTemplate/Maps/VehicleBasic.VehicleBasic
 LocalMapOptions=
 TransitionMap=
 bUseSplitscreen=True
 TwoPlayerSplitscreenLayout=Horizontal
 ThreePlayerSplitscreenLayout=FavorTop
 GameInstanceClass=/Script/Engine.GameInstance
-GameDefaultMap=/Game/VehicleTemplate/Maps/VehicleExampleMap.VehicleExampleMap
+GameDefaultMap=/Game/VehicleTemplate/Maps/VehicleBasic.VehicleBasic
 ServerDefaultMap=/Engine/Maps/Entry.Entry
-GlobalDefaultGameMode=/Script/Engine.GameMode
+GlobalDefaultGameMode=/Game/VehicleTemplate/Blueprints/BP_GM_VehicleBasic.BP_GM_VehicleBasic_C
 GlobalDefaultServerGameMode=None
 
 [/Script/Engine.RendererSettings]
@@ -23,11 +23,25 @@ r.ReflectionMethod=1
 r.GenerateMeshDistanceFields=True
 r.DynamicGlobalIlluminationMethod=1
 r.ClearCoatNormal=True
+r.DefaultFeature.AutoExposure.ExtendDefaultLuminanceRange=True
 r.DefaultFeature.AutoExposure.ExtendDefaultLuminanceRange=true
 r.AllowStaticLighting=False
 
+r.SkinCache.CompileShaders=True
+
+r.RayTracing=True
+
+r.Shadow.Virtual.Enable=1
+
+r.DefaultFeature.LocalExposure.HighlightContrastScale=0.8
+
+r.DefaultFeature.LocalExposure.ShadowContrastScale=0.8
+r.DefaultFeature.LightUnits=2
+r.VirtualTextures=True
+
 [/Script/WindowsTargetPlatform.WindowsTargetSettings]
 DefaultGraphicsRHI=DefaultGraphicsRHI_DX12
+DefaultGraphicsRHI=DefaultGraphicsRHI_DX12
 -D3D12TargetedShaderFormats=PCD3D_SM5
 +D3D12TargetedShaderFormats=PCD3D_SM6
 -D3D11TargetedShaderFormats=PCD3D_SM5
@@ -57,4 +71,29 @@ SoundCueCookQualityIndex=-1
 
 [/Script/LinuxTargetPlatform.LinuxTargetSettings]
 -TargetedRHIs=SF_VULKAN_SM5
-+TargetedRHIs=SF_VULKAN_SM6
\ No newline at end of file
++TargetedRHIs=SF_VULKAN_SM6
+
+[/Script/HardwareTargeting.HardwareTargetingSettings]
+TargetedHardwareClass=Desktop
+AppliedTargetedHardwareClass=Desktop
+DefaultGraphicsPerformance=Maximum
+AppliedDefaultGraphicsPerformance=Maximum
+
+[/Script/Engine.Engine]
++ActiveGameNameRedirects=(OldGameName="TP_VehicleAdvBP",NewGameName="/Script/VehicleTemplateNew")
++ActiveGameNameRedirects=(OldGameName="/Script/TP_VehicleAdvBP",NewGameName="/Script/VehicleTemplateNew")
+
+[/Script/AndroidFileServerEditor.AndroidFileServerRuntimeSettings]
+bEnablePlugin=True
+bAllowNetworkConnection=True
+SecurityToken=0AAA16104B3CEB6E0395D7B28752BFBD
+bIncludeInShipping=False
+bAllowExternalStartInShipping=False
+bCompileAFSProject=False
+bUseCompression=False
+bLogFiles=False
+bReportStats=False
+ConnectionType=USBOnly
+bUseManualIPAddress=False
+ManualIPAddress=
+
diff --git a/Templates/TP_VehicleAdvBP/Config/DefaultGame.ini b/Templates/TP_VehicleAdvBP/Config/DefaultGame.ini
--- a/Templates/TP_VehicleAdvBP/Config/DefaultGame.ini
+++ b/Templates/TP_VehicleAdvBP/Config/DefaultGame.ini
@@ -3,3 +3,6 @@
 ProjectName=Advanced Vehicle Template BP
 
 
+
+[/Script/EngineSettings.GeneralProjectSettings]
+ProjectID=4AE0D26944C3F40803A6CCAEBEC6B72C
diff --git a/Templates/TP_VehicleAdvBP/Config/TemplateDefs.ini b/Templates/TP_VehicleAdvBP/Config/TemplateDefs.ini
--- a/Templates/TP_VehicleAdvBP/Config/TemplateDefs.ini
+++ b/Templates/TP_VehicleAdvBP/Config/TemplateDefs.ini
@@ -45,8 +45,8 @@ ReplacementsInFiles=(Extensions=("cpp","h","ini","cs"),From="%TEMPLATENAME_UPPER
 ReplacementsInFiles=(Extensions=("cpp","h","ini","cs"),From="%TEMPLATENAME_LOWERCASE%",To="%PROJECTNAME_LOWERCASE%",bCaseSensitive=true)
 ReplacementsInFiles=(Extensions=("cpp","h","ini","cs"),From="%TEMPLATENAME%",To="%PROJECTNAME%",bCaseSensitive=false)
 
-
 SharedContentPacks=(MountName="Vehicles",DetailLevels=("Standard"))
-SharedContentPacks=(MountName="Track",DetailLevels=("High"))
-SharedContentPacks=(MountName="LevelPrototyping",DetailLevels=("High"))
-EditDetailLevelPreference="High"
\ No newline at end of file
+EditDetailLevelPreference="High"
+
+Variants=(Name="Timetrial",LocalizedDisplayNames=((Language="en",Text="Timetrial"),(Language="ko",Text="퍼즐")),LocalizedDescriptions=((Language="en",Text="A timed lap with UI and track markers."),(Language="ko",Text="퍼즐 템플릿은 단순한 블록 모눈을 특징으로 합니다.")),SharedContentPacks=((DetailLevels=(Standard),MountName="Variant_Timetrial")) )
+Variants=(Name="Offroad",LocalizedDisplayNames=((Language="en",Text="Offroad"),(Language="ko",Text="퍼즐")),LocalizedDescriptions=((Language="en",Text="An open world map with offroad vehicle."),(Language="ko",Text="퍼즐 템플릿은 단순한 블록 모눈을 특징으로 합니다.")),SharedContentPacks=((DetailLevels=(Standard),MountName="Variant_Offroad")) )
\ No newline at end of file
diff --git a/Templates/TemplateResources/Standard/Variant_Offroad/FeaturePack/manifest.json b/Templates/TemplateResources/Standard/Variant_Offroad/FeaturePack/manifest.json
new file mode 100644
--- /dev/null
+++ b/Templates/TemplateResources/Standard/Variant_Offroad/FeaturePack/manifest.json
@@ -0,0 +1,41 @@
+{
+	"Version": 1,
+	"Name": [
+		{
+			"Language": "en",
+			"Text": "FirstPerson Shared Resource Standard"
+		}
+	],
+	"Description": [
+		{
+			"Language": "en",
+			"Text": "Shared Resource Pack for FirstPerson"
+		}
+	],
+	"AssetTypes": [
+		{
+			"Language": "en",
+			"Text": ""
+		}
+	],
+	"SearchTags": [
+		{
+			"Language": "en",
+			"Text": ""
+		}
+	],
+	"ClassTypes": " ",
+	"Category": "SharedPack",
+	"Thumbnail": "Variant_Shooter.png",
+	"Screenshots": [
+		"Variant_Shooter_Preview.png"
+	],
+	"Ident": "Variant_OffroadStandard",
+	"AdditionalFiles": {
+		"DestinationFilesFolder": "Variant_Offroad",
+		"AdditionalFilesList": [
+			"Templates/TemplateResources/Standard/Variant_Offroad/Content/*.*"
+		]
+	}
+}
+
diff --git a/Templates/TemplateResources/Standard/Variant_Timetrial/FeaturePack/manifest.json b/Templates/TemplateResources/Standard/Variant_Timetrial/FeaturePack/manifest.json
new file mode 100644
--- /dev/null
+++ b/Templates/TemplateResources/Standard/Variant_Timetrial/FeaturePack/manifest.json
@@ -0,0 +1,41 @@
+{
+	"Version": 1,
+	"Name": [
+		{
+			"Language": "en",
+			"Text": "FirstPerson Shared Resource Standard"
+		}
+	],
+	"Description": [
+		{
+			"Language": "en",
+			"Text": "Shared Resource Pack for FirstPerson"
+		}
+	],
+	"AssetTypes": [
+		{
+			"Language": "en",
+			"Text": ""
+		}
+	],
+	"SearchTags": [
+		{
+			"Language": "en",
+			"Text": ""
+		}
+	],
+	"ClassTypes": " ",
+	"Category": "SharedPack",
+	"Thumbnail": "Variant_Shooter.png",
+	"Screenshots": [
+		"Variant_Shooter_Preview.png"
+	],
+	"Ident": "Variant_TimetrialStandard",
+	"AdditionalFiles": {
+		"DestinationFilesFolder": "Variant_Timetrial",
+		"AdditionalFilesList": [
+			"Templates/TemplateResources/Standard/Variant_Timetrial/Content/*.*"
+		]
+	}
+}
+
