Fixed reference to new weapons pack in VR
Fixed timetrial mode in Vehicle
Update vehicle ref in sandbox sim no longer used
Updated dither material on strategy

diff --git a/Templates/TP_VehicleAdvBP/Config/TemplateDefs.ini b/Templates/TP_VehicleAdvBP/Config/TemplateDefs.ini
--- a/Templates/TP_VehicleAdvBP/Config/TemplateDefs.ini
+++ b/Templates/TP_VehicleAdvBP/Config/TemplateDefs.ini
@@ -48,5 +48,10 @@ ReplacementsInFiles=(Extensions=("cpp","h","ini","cs"),From="%TEMPLATENAME%",To=
 SharedContentPacks=(MountName="Vehicles",DetailLevels=("Standard"))
 EditDetailLevelPreference="High"
 
+; Variant packs
+; NOTE: make sure these are commented out for release!
+;SharedContentPacks=(MountName="Variant_Offroad",DetailLevels=("Standard"))
+;SharedContentPacks=(MountName="Variant_Timetrial",DetailLevels=("Standard"))
+
 Variants=(Name="Timetrial",LocalizedDisplayNames=((Language="en",Text="Timetrial"),(Language="ko",Text="퍼즐")),LocalizedDescriptions=((Language="en",Text="A timed lap with UI and track markers."),(Language="ko",Text="퍼즐 템플릿은 단순한 블록 모눈을 특징으로 합니다.")),SharedContentPacks=((DetailLevels=(Standard),MountName="Variant_Timetrial")) )
 Variants=(Name="Offroad",LocalizedDisplayNames=((Language="en",Text="Offroad"),(Language="ko",Text="퍼즐")),LocalizedDescriptions=((Language="en",Text="An open world map with offroad vehicle."),(Language="ko",Text="퍼즐 템플릿은 단순한 블록 모눈을 특징으로 합니다.")),SharedContentPacks=((DetailLevels=(Standard),MountName="Variant_Offroad")) )
\ No newline at end of file
diff --git a/Templates/TP_VirtualRealityBP/Config/TemplateDefs.ini b/Templates/TP_VirtualRealityBP/Config/TemplateDefs.ini
--- a/Templates/TP_VirtualRealityBP/Config/TemplateDefs.ini
+++ b/Templates/TP_VirtualRealityBP/Config/TemplateDefs.ini
@@ -43,7 +43,7 @@ ReplacementsInFiles=(Extensions=("cpp","h","ini","cs"),From="%TEMPLATENAME_UPPER
 ReplacementsInFiles=(Extensions=("cpp","h","ini","cs"),From="%TEMPLATENAME_LOWERCASE%",To="%PROJECTNAME_LOWERCASE%",bCaseSensitive=true)
 ReplacementsInFiles=(Extensions=("cpp","h","ini","cs"),From="%TEMPLATENAME%",To="%PROJECTNAME%",bCaseSensitive=false)
 SortKey=VirtualRealityBP
-SharedContentPacks=(MountName="FPWeapon",DetailLevels=("Standard"))
+SharedContentPacks=(MountName="Weapons",DetailLevels=("Standard"))
 SharedContentPacks=(MountName="VRSpectator",DetailLevels=("Standard"))
 SharedContentPacks=(MountName="LevelPrototyping",DetailLevels=("High"))
 EditDetailLevelPreference="High"
