Templates - Touch Controls for Vehicle Templates

diff --git a/Templates/TP_VehicleAdv/Config/TemplateDefs.ini b/Templates/TP_VehicleAdv/Config/TemplateDefs.ini
--- a/Templates/TP_VehicleAdv/Config/TemplateDefs.ini
+++ b/Templates/TP_VehicleAdv/Config/TemplateDefs.ini
@@ -53,4 +53,5 @@ ReplacementsInFiles=(Extensions=("cpp","h","ini","cs"),From="%TEMPLATENAME%",To=
 SortKey=ZVehicle_Adv
 SharedContentPacks=(MountName="Vehicles",DetailLevels=("Standard"))
 SharedContentPacks=(MountName="Track",DetailLevels=("High"))
-SharedContentPacks=(MountName="LevelPrototyping",DetailLevels=("High"))
\ No newline at end of file
+SharedContentPacks=(MountName="LevelPrototyping",DetailLevels=("High"))
+SharedContentPacks=(MountName="Input",DetailLevels=("High"))
\ No newline at end of file
diff --git a/Templates/TP_VehicleAdvBP/Config/DefaultInput.ini b/Templates/TP_VehicleAdvBP/Config/DefaultInput.ini
--- a/Templates/TP_VehicleAdvBP/Config/DefaultInput.ini
+++ b/Templates/TP_VehicleAdvBP/Config/DefaultInput.ini
@@ -110,7 +110,7 @@ FOVScale=0.011110
 DoubleClickTime=0.200000
 DefaultPlayerInputClass=/Script/EnhancedInput.EnhancedPlayerInput
 DefaultInputComponentClass=/Script/EnhancedInput.EnhancedInputComponent
-DefaultTouchInterface=/Engine/MobileResources/HUD/DefaultVirtualJoysticks.DefaultVirtualJoysticks
+DefaultTouchInterface=None
 -ConsoleKeys=Tilde
 +ConsoleKeys=Tilde
 
diff --git a/Templates/TP_VehicleAdvBP/Config/TemplateDefs.ini b/Templates/TP_VehicleAdvBP/Config/TemplateDefs.ini
--- a/Templates/TP_VehicleAdvBP/Config/TemplateDefs.ini
+++ b/Templates/TP_VehicleAdvBP/Config/TemplateDefs.ini
@@ -37,12 +37,13 @@ ReplacementsInFiles=(Extensions=("cpp","h","ini","cs"),From="%TEMPLATENAME_LOWER
 ReplacementsInFiles=(Extensions=("cpp","h","ini","cs"),From="%TEMPLATENAME%",To="%PROJECTNAME%",bCaseSensitive=false)
 
 SharedContentPacks=(MountName="Vehicles",DetailLevels=("Standard"))
+SharedContentPacks=(MountName="Input",DetailLevels=("High"))
 EditDetailLevelPreference="High"
 
 ; Variant packs
 ; NOTE: make sure these are commented out for release!
-;SharedContentPacks=(MountName="Variant_Offroad",DetailLevels=("Standard"))
-;SharedContentPacks=(MountName="Variant_Timetrial",DetailLevels=("Standard"))
+SharedContentPacks=(MountName="Variant_Offroad",DetailLevels=("Standard"))
+SharedContentPacks=(MountName="Variant_Timetrial",DetailLevels=("Standard"))
 
 Variants=(Name="Timetrial",LocalizedDisplayNames=((Language="en",Text="Timetrial")),LocalizedDescriptions=((Language="en",Text="A timed lap with UI and track markers.")),SharedContentPacks=((DetailLevels=(Standard),MountName="Variant_Timetrial")))
 
