Templates - bug fixes
TP_Vehicle - Matched vehicle properties from BP template
TP_VehicleBP - commented incorrect additional packages in TemplateDefs
TP_ThirdPerson - Matched Sidescrolling NPC appearance
TP_ThirdPerson - Sidescrolling character now faces towards camera when turning
TP_ThirdPerson - Matched StateTree style and added debug texts
TP_ThirdPersonBP - Matched Sidescrolling NPC appearance
TP_ThirdPersonBP - Matched StateTree style and added debug texts
TP_ThirdPerson - Matched enemy spawn locations and loadouts from BP template

#rnx

diff --git a/Templates/TP_ThirdPerson/Source/TP_ThirdPerson/Variant_SideScrolling/SideScrollingCharacter.cpp b/Templates/TP_ThirdPerson/Source/TP_ThirdPerson/Variant_SideScrolling/SideScrollingCharacter.cpp
--- a/Templates/TP_ThirdPerson/Source/TP_ThirdPerson/Variant_SideScrolling/SideScrollingCharacter.cpp
+++ b/Templates/TP_ThirdPerson/Source/TP_ThirdPerson/Variant_SideScrolling/SideScrollingCharacter.cpp
@@ -150,7 +150,7 @@ void ASideScrollingCharacter::DoMove(float Forward)
 		ActionValueY = Forward;
 
 		// figure out the movement direction
-		const FVector MoveDir = FVector(1.0f, Forward > 0.0f ? 0.1f : 0.1f, 0.0f);
+		const FVector MoveDir = FVector(1.0f, Forward > 0.0f ? 0.1f : -0.1f, 0.0f);
 
 		// apply the movement input
 		AddMovementInput(MoveDir, Forward);
diff --git a/Templates/TP_VehicleAdvBP/Config/TemplateDefs.ini b/Templates/TP_VehicleAdvBP/Config/TemplateDefs.ini
--- a/Templates/TP_VehicleAdvBP/Config/TemplateDefs.ini
+++ b/Templates/TP_VehicleAdvBP/Config/TemplateDefs.ini
@@ -42,8 +42,8 @@ EditDetailLevelPreference="High"
 
 ; Variant packs
 ; NOTE: make sure these are commented out for release!
-SharedContentPacks=(MountName="Variant_Offroad",DetailLevels=("Standard"))
-SharedContentPacks=(MountName="Variant_Timetrial",DetailLevels=("Standard"))
+;SharedContentPacks=(MountName="Variant_Offroad",DetailLevels=("Standard"))
+;SharedContentPacks=(MountName="Variant_Timetrial",DetailLevels=("Standard"))
 
 Variants=(Name="Timetrial",LocalizedDisplayNames=((Language="en",Text="Timetrial")),LocalizedDescriptions=((Language="en",Text="A timed lap with UI and track markers.")),SharedContentPacks=((DetailLevels=(Standard),MountName="Variant_Timetrial")))
 
